Branching UHT to UHTLite:

//UE5/Main/Engine/Source/Programs/UnrealHeaderTool/...

to //UE5/Main/Engine/Source/Programs/UHTLite/...

UHTLite is a long-term refactor of UHT for performance and maintainability.

diff --git a/Engine/Source/Programs/UHTLite/Private/BaseParser.cpp b/Engine/Source/Programs/UHTLite/Private/BaseParser.cpp
new file mode 100644
--- /dev/null
+++ b/Engine/Source/Programs/UHTLite/Private/BaseParser.cpp
@@ -0,0 +1,1220 @@
+// Copyright Epic Games, Inc. All Rights Reserved.
+
+#include "BaseParser.h"
+#include "UnrealHeaderTool.h"
+#include "UObject/NameTypes.h"
+#include "UObject/ErrorException.h"
+
+#include "ParserHelper.h"
+
+
+namespace
+{
+	namespace EMetadataValueArgument
+	{
+		enum Type
+		{
+			None,
+			Required,
+			Optional
+		};
+	}
+
+	namespace EMetadataValueAction
+	{
+		enum Type
+		{
+			Remove,
+			Add
+		};
+	}
+
+	struct FMetadataValueAction
+	{
+		FMetadataValueAction(const TCHAR* InMapping, const TCHAR* InDefaultValue, EMetadataValueAction::Type InValueAction)
+		: Mapping         (InMapping)
+		, DefaultValue    (InDefaultValue)
+		, ValueAction     (InValueAction)
+		{
+		}
+
+		FName                      Mapping;
+		FString                    DefaultValue;
+		EMetadataValueAction::Type ValueAction;
+	};
+
+	struct FMetadataKeyword
+	{
+		FMetadataKeyword(EMetadataValueArgument::Type InValueArgument)
+		: ValueArgument(InValueArgument)
+		{
+		}
+
+		void InsertAddAction(const TCHAR* InMapping, const TCHAR* InDefaultValue)
+		{
+			ValueActions.Add(FMetadataValueAction(InMapping, InDefaultValue, EMetadataValueAction::Add));
+		}
+
+		void InsertRemoveAction(const TCHAR* InMapping)
+		{
+			ValueActions.Add(FMetadataValueAction(InMapping, TEXT(""), EMetadataValueAction::Remove));
+		}
+
+		void ApplyToMetadata(TMap<FName, FString>& Metadata, const FString* Value = nullptr)
+		{
+			for (const FMetadataValueAction& ValueAction : ValueActions)
+			{
+				if (ValueAction.ValueAction == EMetadataValueAction::Add)
+				{
+					FBaseParser::InsertMetaDataPair(Metadata, ValueAction.Mapping, Value ? *Value : ValueAction.DefaultValue);
+				}
+				else
+				{
+					Metadata.Remove(ValueAction.Mapping);
+				}
+			}
+		}
+
+		TArray<FMetadataValueAction> ValueActions;
+		EMetadataValueArgument::Type ValueArgument;
+	};
+
+	FMetadataKeyword* GetMetadataKeyword(const TCHAR* Keyword)
+	{
+		static TMap<FString, FMetadataKeyword> Dictionary;
+		if (!Dictionary.Num())
+		{
+			FMetadataKeyword& DisplayName = Dictionary.Add(TEXT("DisplayName"), EMetadataValueArgument::Required);
+			DisplayName.InsertAddAction(TEXT("DisplayName"), TEXT(""));
+
+			FMetadataKeyword& FriendlyName = Dictionary.Add(TEXT("FriendlyName"), EMetadataValueArgument::Required);
+			FriendlyName.InsertAddAction(TEXT("FriendlyName"), TEXT(""));
+
+			FMetadataKeyword& BlueprintInternalUseOnly = Dictionary.Add(TEXT("BlueprintInternalUseOnly"), EMetadataValueArgument::None);
+			BlueprintInternalUseOnly.InsertAddAction(TEXT("BlueprintInternalUseOnly"), TEXT("true"));
+			BlueprintInternalUseOnly.InsertAddAction(TEXT("BlueprintType"), TEXT("true"));
+
+			FMetadataKeyword& BlueprintType = Dictionary.Add(TEXT("BlueprintType"), EMetadataValueArgument::None);
+			BlueprintType.InsertAddAction(TEXT("BlueprintType"), TEXT("true"));
+
+			FMetadataKeyword& NotBlueprintType = Dictionary.Add(TEXT("NotBlueprintType"), EMetadataValueArgument::None);
+			NotBlueprintType.InsertAddAction(TEXT("NotBlueprintType"), TEXT("true"));
+			NotBlueprintType.InsertRemoveAction(TEXT("BlueprintType"));
+
+			FMetadataKeyword& Blueprintable = Dictionary.Add(TEXT("Blueprintable"), EMetadataValueArgument::None);
+			Blueprintable.InsertAddAction(TEXT("IsBlueprintBase"), TEXT("true"));
+			Blueprintable.InsertAddAction(TEXT("BlueprintType"),   TEXT("true"));
+
+			FMetadataKeyword& CallInEditor = Dictionary.Add(TEXT("CallInEditor"), EMetadataValueArgument::None);
+			CallInEditor.InsertAddAction(TEXT("CallInEditor"), TEXT("true"));
+
+			FMetadataKeyword& NotBlueprintable = Dictionary.Add(TEXT("NotBlueprintable"), EMetadataValueArgument::None);
+			NotBlueprintable.InsertAddAction   (TEXT("IsBlueprintBase"), TEXT("false"));
+			NotBlueprintable.InsertRemoveAction(TEXT("BlueprintType"));
+
+			FMetadataKeyword& Category = Dictionary.Add(TEXT("Category"), EMetadataValueArgument::Required);
+			Category.InsertAddAction(TEXT("Category"), TEXT(""));
+
+			FMetadataKeyword& ExperimentalFeature = Dictionary.Add(TEXT("Experimental"), EMetadataValueArgument::None);
+			ExperimentalFeature.InsertAddAction(TEXT("DevelopmentStatus"), TEXT("Experimental"));
+
+			FMetadataKeyword& EarlyAccessFeature = Dictionary.Add(TEXT("EarlyAccessPreview"), EMetadataValueArgument::None);
+			EarlyAccessFeature.InsertAddAction(TEXT("DevelopmentStatus"), TEXT("EarlyAccess"));
+
+			FMetadataKeyword& DocumentationPolicy = Dictionary.Add(TEXT("DocumentationPolicy"), EMetadataValueArgument::None);
+			DocumentationPolicy.InsertAddAction(TEXT("DocumentationPolicy"), TEXT("Strict"));
+
+			FMetadataKeyword& SparseClassDataType = Dictionary.Add(TEXT("SparseClassDataType"), EMetadataValueArgument::Required);
+			SparseClassDataType.InsertAddAction(TEXT("SparseClassDataType"), TEXT(""));
+		}
+
+		return Dictionary.Find(Keyword);
+	}
+}
+
+//////////////////////////////////////////////////////////////////////////
+// FPropertySpecifier
+
+FString FPropertySpecifier::ConvertToString() const
+{
+	FString Result;
+
+	// Emit the specifier key
+	Result += Key;
+
+	// Emit the values if there are any
+	if (Values.Num())
+	{
+		Result += TEXT("=");
+
+		if (Values.Num() == 1)
+		{
+			// One value goes on it's own
+			Result += Values[0];
+		}
+		else
+		{
+			// More than one value goes in parens, separated by commas
+			Result += TEXT("(");
+			for (int32 ValueIndex = 0; ValueIndex < Values.Num(); ++ValueIndex)
+			{
+				if (ValueIndex > 0)
+				{
+					Result += TEXT(", ");
+				}
+				Result += Values[ValueIndex];
+			}
+			Result += TEXT(")");
+		}
+	}
+
+	return Result;
+}
+
+/////////////////////////////////////////////////////
+// FBaseParser
+
+FBaseParser::FBaseParser()
+	: StatementsParsed(0)
+	, LinesParsed(0)
+{
+}
+
+void FBaseParser::ResetParser(const TCHAR* SourceBuffer, int32 StartingLineNumber)
+{
+	Input = SourceBuffer;
+	InputLen = FCString::Strlen(Input);
+	InputPos = 0;
+	PrevPos = 0;
+	PrevLine = 1;
+	InputLine = StartingLineNumber;
+}
+
+/*-----------------------------------------------------------------------------
+	Single-character processing.
+-----------------------------------------------------------------------------*/
+
+//
+// Get a single character from the input stream and return it, or 0=end.
+//
+TCHAR FBaseParser::GetChar(bool bLiteral)
+{
+	bool bInsideComment = false;
+
+	PrevPos = InputPos;
+	PrevLine = InputLine;
+
+Loop:
+	const TCHAR c = Input[InputPos++];
+	if (bInsideComment)
+	{
+		// Record the character as a comment.
+		PrevComment += c;
+	}
+
+	if (c == TEXT('\n'))
+	{
+		InputLine++;
+	}
+	else if (!bLiteral)
+	{
+		const TCHAR NextChar = PeekChar();
+		if ( c==TEXT('/') && NextChar==TEXT('*') )
+		{
+			if (!bInsideComment)
+			{
+				ClearComment();
+				// Record the slash and star.
+				PrevComment += c;
+				PrevComment += NextChar;
+				bInsideComment = true;
+
+				// Move past the star. Do it only when not in comment,
+				// otherwise end of comment might be missed e.g.
+				// /*/ Comment /*/
+				// ~~~~~~~~~~~~~^ Will report second /* as beginning of comment
+				// And throw error that end of file is found in comment.
+				InputPos++;
+			}
+
+			goto Loop;
+		}
+		else if( c==TEXT('*') && NextChar==TEXT('/') )
+		{
+			if (!bInsideComment)
+			{
+				ClearComment();
+				FError::Throwf(TEXT("Unexpected '*/' outside of comment") );
+			}
+
+			/** Asterisk and slash always end comment. */
+			bInsideComment = false;
+
+			// Star already recorded; record the slash.
+			PrevComment += Input[InputPos];
+
+			InputPos++;
+			goto Loop;
+		}
+	}
+
+	if (bInsideComment)
+	{
+		if (c == 0)
+		{
+			ClearComment();
+			FError::Throwf(TEXT("End of class header encountered inside comment") );
+		}
+		goto Loop;
+	}
+	return c;
+}
+
+//
+// Unget the previous character retrieved with GetChar().
+//
+void FBaseParser::UngetChar()
+{
+	InputPos = PrevPos;
+	InputLine = PrevLine;
+}
+
+//
+// Look at a single character from the input stream and return it, or 0=end.
+// Has no effect on the input stream.
+//
+TCHAR FBaseParser::PeekChar()
+{
+	return (InputPos < InputLen) ? Input[InputPos] : 0;
+}
+
+//
+// Skip past all spaces and tabs in the input stream.
+//
+TCHAR FBaseParser::GetLeadingChar()
+{
+	TCHAR TrailingCommentNewline = 0;
+
+	for (;;)
+	{
+		bool MultipleNewlines = false;
+
+		TCHAR c;
+
+		// Skip blanks.
+		do
+		{
+			c = GetChar();
+
+			// Check if we've encountered another newline since the last one
+			if (c == TrailingCommentNewline)
+			{
+				MultipleNewlines = true;
+			}
+		} while (IsWhitespace(c));
+
+		if (c != TEXT('/') || PeekChar() != TEXT('/'))
+		{
+			return c;
+		}
+
+		// Clear the comment if we've encountered newlines since the last comment
+		if (MultipleNewlines)
+		{
+			ClearComment();
+		}
+
+		// Record the first slash.  The first iteration of the loop will get the second slash.
+		PrevComment += c;
+
+		do
+		{
+			c = GetChar(true);
+			if (c == 0)
+				return c;
+			PrevComment += c;
+		} while (!IsEOL(c));
+
+		TrailingCommentNewline = c;
+
+		for (;;)
+		{
+			c = GetChar();
+			if (c == 0)
+				return c;
+			if (c == TrailingCommentNewline || !IsEOL(c))
+			{
+				UngetChar();
+				break;
+			}
+
+			PrevComment += c;
+		}
+	}
+}
+
+bool FBaseParser::IsEOL( TCHAR c )
+{
+	return c==TEXT('\n') || c==TEXT('\r') || c==0;
+}
+
+bool FBaseParser::IsWhitespace( TCHAR c )
+{
+	return c==TEXT(' ') || c==TEXT('\t') || c==TEXT('\r') || c==TEXT('\n');
+}
+
+/*-----------------------------------------------------------------------------
+	Tokenizing.
+-----------------------------------------------------------------------------*/
+
+// Gets the next token from the input stream, advancing the variables which keep track of the current input position and line.
+bool FBaseParser::GetToken( FToken& Token, bool bNoConsts/*=false*/, ESymbolParseOption bParseTemplateClosingBracket/*=ESymbolParseOption::Normal*/ )
+{
+	Token.ClearTokenName();
+	TCHAR c = GetLeadingChar();
+	if( c == 0 )
+	{
+		UngetChar();
+		return 0;
+	}
+	TCHAR p = PeekChar();
+	Token.StartPos		= PrevPos;
+	Token.StartLine		= PrevLine;
+	if( (c>='A' && c<='Z') || (c>='a' && c<='z') || (c=='_') )
+	{
+		// Alphanumeric token.
+		int32 Length=0;
+		do
+		{
+			Token.Identifier[Length++] = c;
+			if( Length >= NAME_SIZE )
+			{
+				FError::Throwf(TEXT("Identifer length exceeds maximum of %i"), (int32)NAME_SIZE);
+				Length = ((int32)NAME_SIZE) - 1;
+				break;
+			}
+			c = GetChar();
+		} while( ((c>='A')&&(c<='Z')) || ((c>='a')&&(c<='z')) || ((c>='0')&&(c<='9')) || (c=='_') );
+		UngetChar();
+		Token.Identifier[Length]=0;
+
+		// Assume this is an identifier unless we find otherwise.
+		Token.TokenType = TOKEN_Identifier;
+
+		// If const values are allowed, determine whether the identifier represents a constant
+		if ( !bNoConsts )
+		{
+			// See if the identifier is part of a vector, rotation or other struct constant.
+			// boolean true/false
+			if( Token.Matches(TEXT("true"), ESearchCase::IgnoreCase) )
+			{
+				Token.SetConstBool(true);
+				return true;
+			}
+			else if( Token.Matches(TEXT("false"), ESearchCase::IgnoreCase) )
+			{
+				Token.SetConstBool(false);
+				return true;
+			}
+		}
+		return true;
+	}
+
+	// if const values are allowed, determine whether the non-identifier token represents a const
+	else if ( !bNoConsts && ((c>='0' && c<='9') || ((c=='+' || c=='-') && (p>='0' && p<='9'))) )
+	{
+		// Integer or floating point constant.
+		bool  bIsFloat = 0;
+		int32 Length   = 0;
+		bool  bIsHex   = 0;
+		do
+		{
+			if( c==TEXT('.') )
+			{
+				bIsFloat = true;
+			}
+			if( c==TEXT('X') || c == TEXT('x') )
+			{
+				bIsHex = true;
+			}
+
+			Token.Identifier[Length++] = c;
+			if( Length >= NAME_SIZE )
+			{
+				FError::Throwf(TEXT("Number length exceeds maximum of %i "), (int32)NAME_SIZE );
+				Length = ((int32)NAME_SIZE) - 1;
+				break;
+			}
+			c = FChar::ToUpper(GetChar());
+		} while ((c >= TEXT('0') && c <= TEXT('9')) || (!bIsFloat && c == TEXT('.')) || (!bIsHex && c == TEXT('X')) || (bIsHex && c >= TEXT('A') && c <= TEXT('F')));
+
+		Token.Identifier[Length]=0;
+		if (!bIsFloat || c != 'F')
+		{
+			UngetChar();
+		}
+
+		if (bIsFloat)
+		{
+			Token.SetConstFloat( FCString::Atof(Token.Identifier) );
+		}
+		else if (bIsHex)
+		{
+			TCHAR* End = Token.Identifier + FCString::Strlen(Token.Identifier);
+			Token.SetConstInt64( FCString::Strtoi64(Token.Identifier,&End,0) );
+		}
+		else
+		{
+			Token.SetConstInt64( FCString::Atoi64(Token.Identifier) );
+		}
+		return true;
+	}
+	else if (c == '\'')
+	{
+		TCHAR ActualCharLiteral = GetChar(/*bLiteral=*/ true);
+
+		if (ActualCharLiteral == '\\')
+		{
+			ActualCharLiteral = GetChar(/*bLiteral=*/ true);
+			switch (ActualCharLiteral)
+			{
+			case TCHAR('t'):
+				ActualCharLiteral = '\t';
+				break;
+			case TCHAR('n'):
+				ActualCharLiteral = '\n';
+				break;
+			case TCHAR('r'):
+				ActualCharLiteral = '\r';
+				break;
+			}
+		}
+
+		c = GetChar(/*bLiteral=*/ true);
+		if (c != '\'')
+		{
+			FError::Throwf(TEXT("Unterminated character constant"));
+			UngetChar();
+		}
+
+		Token.SetConstChar(ActualCharLiteral);
+		return true;
+	}
+	else if (c == '"')
+	{
+		// String constant.
+		TCHAR Temp[MAX_STRING_CONST_SIZE];
+		int32 Length=0;
+		c = GetChar(/*bLiteral=*/ true);
+		while( (c!='"') && !IsEOL(c) )
+		{
+			if( c=='\\' )
+			{
+				c = GetChar(/*bLiteral=*/ true);
+				if( IsEOL(c) )
+				{
+					break;
+				}
+				else if(c == 'n')
+				{
+					// Newline escape sequence.
+					c = '\n';
+				}
+			}
+			Temp[Length++] = c;
+			if( Length >= MAX_STRING_CONST_SIZE )
+			{
+				FError::Throwf(TEXT("String constant exceeds maximum of %i characters"), (int32)MAX_STRING_CONST_SIZE );
+				c = TEXT('\"');
+				Length = ((int32)MAX_STRING_CONST_SIZE) - 1;
+				break;
+			}
+			c = GetChar(/*bLiteral=*/ true);
+		}
+		Temp[Length]=0;
+
+		if( c != '"' )
+		{
+			FError::Throwf(TEXT("Unterminated string constant: %s"), Temp);
+			UngetChar();
+		}
+
+		Token.SetConstString(Temp);
+		return true;
+	}
+	else
+	{
+		// Symbol.
+		int32 Length=0;
+		Token.Identifier[Length++] = c;
+
+		// Handle special 2-character symbols.
+		#define PAIR(cc,dd) ((c==cc)&&(d==dd)) /* Comparison macro for convenience */
+		TCHAR d = GetChar();
+		if
+		(	PAIR('<','<')
+		||	(PAIR('>','>') && (bParseTemplateClosingBracket != ESymbolParseOption::CloseTemplateBracket))
+		||	PAIR('!','=')
+		||	PAIR('<','=')
+		||	PAIR('>','=')
+		||	PAIR('+','+')
+		||	PAIR('-','-')
+		||	PAIR('+','=')
+		||	PAIR('-','=')
+		||	PAIR('*','=')
+		||	PAIR('/','=')
+		||	PAIR('&','&')
+		||	PAIR('|','|')
+		||	PAIR('^','^')
+		||	PAIR('=','=')
+		||	PAIR('*','*')
+		||	PAIR('~','=')
+		||	PAIR(':',':')
+		)
+		{
+			Token.Identifier[Length++] = d;
+			if( c=='>' && d=='>' )
+			{
+				if (GetChar()=='>')
+				{
+					Token.Identifier[Length++] = '>';
+				}
+				else
+				{
+					UngetChar();
+				}
+			}
+		}
+		else
+		{
+			UngetChar();
+		}
+		#undef PAIR
+
+		Token.Identifier[Length] = 0;
+		Token.TokenType = TOKEN_Symbol;
+
+		return true;
+	}
+}
+
+bool FBaseParser::GetRawTokenRespectingQuotes( FToken& Token, TCHAR StopChar /* = TCHAR('\n') */ )
+{
+	// Get token after whitespace.
+	TCHAR Temp[MAX_STRING_CONST_SIZE];
+	int32 Length=0;
+	TCHAR c = GetLeadingChar();
+
+	bool bInQuote = false;
+
+	while( !IsEOL(c) && ((c != StopChar) || bInQuote) )
+	{
+		if( (c=='/' && PeekChar()=='/') || (c=='/' && PeekChar()=='*') )
+		{
+			break;
+		}
+
+		if (c == '"')
+		{
+			bInQuote = !bInQuote;
+		}
+
+		Temp[Length++] = c;
+		if( Length >= MAX_STRING_CONST_SIZE )
+		{
+			FError::Throwf(TEXT("Identifier exceeds maximum of %i characters"), (int32)MAX_STRING_CONST_SIZE );
+			c = GetChar(true);
+			Length = ((int32)MAX_STRING_CONST_SIZE) - 1;
+			break;
+		}
+		c = GetChar(true);
+	}
+	UngetChar();
+
+	if (bInQuote)
+	{
+		FError::Throwf(TEXT("Unterminated quoted string"));
+	}
+
+	// Get rid of trailing whitespace.
+	while( Length>0 && (Temp[Length-1]==' ' || Temp[Length-1]==9 ) )
+	{
+		Length--;
+	}
+	Temp[Length]=0;
+
+	Token.SetConstString(Temp);
+
+	return Length>0;
+}
+
+/**
+ * Put all text from the current position up to either EOL or the StopToken
+ * into Token.  Advances the compiler's current position.
+ *
+ * @param	Token	[out] will contain the text that was parsed
+ * @param	StopChar	stop processing when this character is reached
+ *
+ * @return	the number of character parsed
+ */
+bool FBaseParser::GetRawToken( FToken& Token, TCHAR StopChar /* = TCHAR('\n') */ )
+{
+	// Get token after whitespace.
+	TCHAR Temp[MAX_STRING_CONST_SIZE];
+	int32 Length=0;
+	TCHAR c = GetLeadingChar();
+	while( !IsEOL(c) && c != StopChar )
+	{
+		if( (c=='/' && PeekChar()=='/') || (c=='/' && PeekChar()=='*') )
+		{
+			break;
+		}
+		Temp[Length++] = c;
+		if( Length >= MAX_STRING_CONST_SIZE )
+		{
+			FError::Throwf(TEXT("Identifier exceeds maximum of %i characters"), (int32)MAX_STRING_CONST_SIZE );
+		}
+		c = GetChar(true);
+	}
+	UngetChar();
+
+	// Get rid of trailing whitespace.
+	while( Length>0 && (Temp[Length-1]==' ' || Temp[Length-1]==9 ) )
+	{
+		Length--;
+	}
+	Temp[Length]=0;
+
+	Token.SetConstString(Temp);
+
+	return Length>0;
+}
+
+//
+// Get an identifier token, return 1 if gotten, 0 if not.
+//
+bool FBaseParser::GetIdentifier( FToken& Token, bool bNoConsts )
+{
+	if (!GetToken(Token, bNoConsts))
+	{
+		return false;
+	}
+
+	if (Token.TokenType == TOKEN_Identifier)
+	{
+		return true;
+	}
+
+	UngetToken(Token);
+	return false;
+}
+
+//
+// Get a symbol token, return 1 if gotten, 0 if not.
+//
+bool FBaseParser::GetSymbol( FToken& Token )
+{
+	if (!GetToken(Token))
+	{
+		return false;
+	}
+
+	if( Token.TokenType == TOKEN_Symbol )
+	{
+		return true;
+	}
+
+	UngetToken(Token);
+	return false;
+}
+
+bool FBaseParser::GetConstInt(int32& Result, const TCHAR* Tag)
+{
+	FToken Token;
+	if (GetToken(Token))
+	{
+		if (Token.GetConstInt(Result))
+		{
+			return true;
+		}
+		else
+		{
+			UngetToken(Token);
+		}
+	}
+
+	if (Tag != NULL)
+	{
+		FError::Throwf(TEXT("%s: Missing constant integer"), Tag );
+	}
+
+	return false;
+}
+
+bool FBaseParser::GetConstInt64(int64& Result, const TCHAR* Tag)
+{
+	FToken Token;
+	if (GetToken(Token))
+	{
+		if (Token.GetConstInt64(Result))
+		{
+			return true;
+		}
+		else
+		{
+			UngetToken(Token);
+		}
+	}
+
+	if (Tag != NULL)
+	{
+		FError::Throwf(TEXT("%s: Missing constant integer"), Tag );
+	}
+
+	return false;
+}
+
+bool FBaseParser::MatchSymbol( const TCHAR Match, ESymbolParseOption bParseTemplateClosingBracket/*=ESymbolParseOption::Normal*/ )
+{
+	FToken Token;
+
+	if (GetToken(Token, /*bNoConsts=*/ true, bParseTemplateClosingBracket))
+	{
+		if (Token.TokenType==TOKEN_Symbol && Token.Identifier[0] == Match && Token.Identifier[1] == 0)
+		{
+			return true;
+		}
+		else
+		{
+			UngetToken(Token);
+		}
+	}
+
+	return false;
+}
+
+bool FBaseParser::MatchSymbol(const TCHAR* Match, ESymbolParseOption bParseTemplateClosingBracket/*=ESymbolParseOption::Normal*/)
+{
+	FToken Token;
+
+	if (GetToken(Token, /*bNoConsts=*/ true, bParseTemplateClosingBracket))
+	{
+		if (Token.TokenType==TOKEN_Symbol && !FCString::Strcmp(Token.Identifier, Match))
+		{
+			return true;
+		}
+		else
+		{
+			UngetToken(Token);
+		}
+	}
+
+	return false;
+}
+
+//
+// Get a specific identifier and return 1 if gotten, 0 if not.
+// This is used primarily for checking for required symbols during compilation.
+//
+bool FBaseParser::MatchIdentifierByName( FName Match )
+{
+	FToken Token;
+	if (!GetToken(Token))
+	{
+		return false;
+	}
+
+	if ((Token.TokenType == TOKEN_Identifier) && (Token.GetTokenName() == Match))
+	{
+		return true;
+	}
+
+	UngetToken(Token);
+	return false;
+}
+
+bool FBaseParser::MatchIdentifier( const TCHAR* Match, ESearchCase::Type SearchCase)
+{
+	FToken Token;
+	if (GetToken(Token))
+	{
+		if (Token.TokenType==TOKEN_Identifier && ((SearchCase == ESearchCase::CaseSensitive) ? !FCString::Strcmp(Token.Identifier, Match) : !FCString::Stricmp(Token.Identifier, Match)))
+		{
+			return true;
+		}
+		else
+		{
+			UngetToken(Token);
+		}
+	}
+	
+	return false;
+}
+
+bool FBaseParser::MatchConstInt( const TCHAR* Match )
+{
+	FToken Token;
+	if (GetToken(Token))
+	{
+		if( Token.TokenType==TOKEN_Const && (Token.Type == CPT_Int || Token.Type == CPT_Int64) && FCString::Strcmp(Token.Identifier,Match)==0 )
+		{
+			return true;
+		}
+		else
+		{
+			UngetToken(Token);
+		}
+	}
+	
+	return false;
+}
+
+bool FBaseParser::MatchAnyConstInt()
+{
+	FToken Token;
+	if (GetToken(Token))
+	{
+		if( Token.TokenType==TOKEN_Const && (Token.Type == CPT_Int || Token.Type == CPT_Int64) )
+		{
+			return true;
+		}
+		else
+		{
+			UngetToken(Token);
+		}
+	}
+
+	return false;
+}
+
+void FBaseParser::MatchSemi()
+{
+	if( !MatchSymbol(TEXT(';')) )
+	{
+		FToken Token;
+		if( GetToken(Token) )
+		{
+			FError::Throwf(TEXT("Missing ';' before '%s'"), Token.Identifier );
+		}
+		else
+		{
+			FError::Throwf(TEXT("Missing ';'") );
+		}
+	}
+}
+
+
+//
+// Peek ahead and see if a symbol follows in the stream.
+//
+bool FBaseParser::PeekSymbol( const TCHAR Match )
+{
+	FToken Token;
+	if (!GetToken(Token, true))
+	{
+		return false;
+	}
+	UngetToken(Token);
+
+	return Token.TokenType==TOKEN_Symbol && Token.Identifier[0] == Match && Token.Identifier[1] == 0;
+}
+
+//
+// Peek ahead and see if an identifier follows in the stream.
+//
+bool FBaseParser::PeekIdentifierByName( FName Match )
+{
+	FToken Token;
+	if (!GetToken(Token, true))
+	{
+		return false;
+	}
+	UngetToken(Token);
+	return Token.TokenType==TOKEN_Identifier && Token.GetTokenName()==Match;
+}
+
+bool FBaseParser::PeekIdentifier( const TCHAR* Match, ESearchCase::Type SearchCase)
+{
+	FToken Token;
+	if (!GetToken(Token, true))
+	{
+		return false;
+	}
+	UngetToken(Token);
+	return Token.TokenType==TOKEN_Identifier && ((SearchCase == ESearchCase::CaseSensitive) ? !FCString::Strcmp(Token.Identifier, Match) : !FCString::Stricmp(Token.Identifier, Match));
+}
+
+//
+// Unget the most recently gotten token.
+//
+void FBaseParser::UngetToken( const FToken& Token )
+{
+	InputPos = Token.StartPos;
+	InputLine = Token.StartLine;
+}
+
+//
+// Require a symbol.
+//
+void FBaseParser::RequireSymbol( const TCHAR Match, const TCHAR* Tag, ESymbolParseOption bParseTemplateClosingBracket/*=ESymbolParseOption::Normal*/ )
+{
+	if (!MatchSymbol(Match, bParseTemplateClosingBracket))
+	{
+		FError::Throwf(TEXT("Missing '%c' in %s"), Match, Tag );
+	}
+}
+
+void FBaseParser::RequireSymbol(const TCHAR Match, TFunctionRef<FString()> ErrorGetter, ESymbolParseOption bParseTemplateClosingBracket/*=ESymbolParseOption::Normal*/)
+{
+	if (!MatchSymbol(Match, bParseTemplateClosingBracket))
+	{
+		FError::Throwf(TEXT("Missing '%c' in %s"), Match, *ErrorGetter());
+	}
+}
+
+//
+// Require an integer.
+//
+void FBaseParser::RequireConstInt( const TCHAR* Match, const TCHAR* Tag )
+{
+	if (!MatchConstInt(Match))
+	{
+		FError::Throwf(TEXT("Missing integer '%s' in %s"), Match, Tag );
+	}
+}
+
+void FBaseParser::RequireAnyConstInt( const TCHAR* Tag )
+{
+	if (!MatchAnyConstInt())
+	{
+		FError::Throwf(TEXT("Missing integer in %s"), Tag );
+	}
+}
+
+//
+// Require an identifier.
+//
+void FBaseParser::RequireIdentifier( const TCHAR* Match, const ESearchCase::Type SearchCase, const TCHAR* Tag )
+{
+	if (!MatchIdentifier(Match, SearchCase))
+	{
+		FError::Throwf(TEXT("Missing '%s' in %s"), Match, Tag );
+	}
+}
+
+// Clears out the stored comment.
+void FBaseParser::ClearComment()
+{
+	PrevComment.Reset();
+}
+
+// Reads a new-style value
+//@TODO: UCREMOVAL: Needs a better name
+FString FBaseParser::ReadNewStyleValue(const TCHAR* TypeOfSpecifier)
+{
+	FToken ValueToken;
+	if (!GetToken(ValueToken, false))
+	{
+		FError::Throwf(TEXT("Expected a value when handling a %s"), TypeOfSpecifier);
+	}
+
+	FString Result;
+
+	switch (ValueToken.TokenType)
+	{
+		case TOKEN_Identifier:
+		case TOKEN_Symbol:
+
+			Result = ValueToken.Identifier;
+
+			if (MatchSymbol(TEXT('=')))
+			{
+				Result += TEXT("=");
+				Result += ReadNewStyleValue(TypeOfSpecifier);
+			}
+			break;
+
+		case TOKEN_Const:
+
+			Result = ValueToken.GetConstantValue();
+			break;
+	}
+
+	return Result;
+}
+
+// Reads [ Value | ['(' Value [',' Value]* ')'] ] and places each value into the Items array
+bool FBaseParser::ReadOptionalCommaSeparatedListInParens(TArray<FString>& Items, const TCHAR* TypeOfSpecifier)
+{
+	if (MatchSymbol(TEXT('(')))
+	{
+		do 
+		{
+			FString Value = ReadNewStyleValue(TypeOfSpecifier);
+			Items.Add(Value);
+		} while ( MatchSymbol(TEXT(',')) );
+
+		RequireSymbol(TEXT(')'), TypeOfSpecifier);
+
+		return true;
+	}
+
+	return false;
+}
+
+void FBaseParser::ParseNameWithPotentialAPIMacroPrefix(FString& DeclaredName, FString& RequiredAPIMacroIfPresent, const TCHAR* FailureMessage)
+{
+	// Expecting Name | (MODULE_API Name)
+	FToken NameToken;
+
+	// Read an identifier
+	if (!GetIdentifier(NameToken))
+	{
+		FError::Throwf(TEXT("Missing %s name"), FailureMessage);
+	}
+
+	// Is the identifier the name or an DLL import/export API macro?
+	FString NameTokenStr = NameToken.Identifier;
+	if (NameTokenStr.EndsWith(TEXT("_API"), ESearchCase::CaseSensitive))
+	{
+		RequiredAPIMacroIfPresent = MoveTemp(NameTokenStr);
+
+		// Read the real name
+		if (!GetIdentifier(NameToken))
+		{
+			FError::Throwf(TEXT("Missing %s name"), FailureMessage);
+		}
+		DeclaredName = NameToken.Identifier;
+	}
+	else
+	{
+		DeclaredName = MoveTemp(NameTokenStr);
+		RequiredAPIMacroIfPresent.Reset();
+	}
+}
+
+// Reads a set of specifiers (with optional values) inside the () of a new-style metadata macro like UPROPERTY or UFUNCTION
+void FBaseParser::ReadSpecifierSetInsideMacro(TArray<FPropertySpecifier>& SpecifiersFound, const TCHAR* TypeOfSpecifier, TMap<FName, FString>& MetaData)
+{
+	int32 FoundSpecifierCount = 0;
+
+	auto ErrorMessageGetter = [TypeOfSpecifier]() { return FString::Printf(TEXT("%s declaration specifier"), TypeOfSpecifier); };
+
+	RequireSymbol(TEXT('('), ErrorMessageGetter);
+
+	while (!MatchSymbol(TEXT(')')))
+	{
+		if (FoundSpecifierCount > 0)
+		{
+			RequireSymbol(TEXT(','), ErrorMessageGetter);
+		}
+		++FoundSpecifierCount;
+
+		// Read the specifier key
+		FToken Specifier;
+		if (!GetToken(Specifier))
+		{
+			FError::Throwf(TEXT("Expected %s"), *ErrorMessageGetter());
+		}
+
+		if (Specifier.Matches(TEXT("meta"), ESearchCase::IgnoreCase))
+		{
+			RequireSymbol(TEXT('='), ErrorMessageGetter);
+			RequireSymbol(TEXT('('), ErrorMessageGetter);
+
+			// Keep reading comma-separated metadata pairs
+			do 
+			{
+				// Read a key
+				FToken MetaKeyToken;
+				if (!GetIdentifier(MetaKeyToken))
+				{
+					FError::Throwf(TEXT("Expected a metadata key"));
+				}
+
+				FString Key = MetaKeyToken.Identifier;
+
+				// Potentially read a value
+				FString Value;
+				if (MatchSymbol(TEXT('=')))
+				{
+					Value = ReadNewStyleValue(TypeOfSpecifier);
+				}
+
+				// Validate the value is a valid type for the key and insert it into the map
+				InsertMetaDataPair(MetaData, MoveTemp(Key), MoveTemp(Value));
+			} while ( MatchSymbol(TEXT(',')) );
+
+			RequireSymbol(TEXT(')'), ErrorMessageGetter);
+		}
+		// Look up specifier in metadata dictionary
+		else if (FMetadataKeyword* MetadataKeyword = GetMetadataKeyword(Specifier.Identifier))
+		{
+			if (MatchSymbol(TEXT('=')))
+			{
+				if (MetadataKeyword->ValueArgument == EMetadataValueArgument::None)
+				{
+					FError::Throwf(TEXT("Incorrect = after metadata specifier '%s'"), Specifier.Identifier);
+				}
+
+				FString Value = ReadNewStyleValue(TypeOfSpecifier);
+				MetadataKeyword->ApplyToMetadata(MetaData, &Value);
+			}
+			else
+			{
+				if (MetadataKeyword->ValueArgument == EMetadataValueArgument::Required)
+				{
+					FError::Throwf(TEXT("Missing = after metadata specifier '%s'"), Specifier.Identifier);
+				}
+
+				MetadataKeyword->ApplyToMetadata(MetaData);
+			}
+		}
+		else
+		{
+			// Creating a new specifier
+			SpecifiersFound.Emplace(Specifier.Identifier);
+
+			// Look for a value for this specifier
+			if (MatchSymbol(TEXT('=')) || PeekSymbol(TEXT('(')))
+			{
+				TArray<FString>& NewPairValues = SpecifiersFound.Last().Values;
+				if (!ReadOptionalCommaSeparatedListInParens(NewPairValues, TypeOfSpecifier))
+				{
+					NewPairValues.Add(ReadNewStyleValue(TypeOfSpecifier));
+				}
+			}
+		}
+	}
+}
+
+void FBaseParser::InsertMetaDataPair(TMap<FName, FString>& MetaData, FString Key, FString Value)
+{
+	// make sure the key is valid
+	if (Key.Len() == 0)
+	{
+		FError::Throwf(TEXT("Invalid metadata"));
+	}
+
+	// trim extra white space and quotes
+	Key.TrimStartAndEndInline();
+
+	InsertMetaDataPair(MetaData, FName(*Key), MoveTemp(Value));
+}
+
+void FBaseParser::InsertMetaDataPair(TMap<FName, FString>& MetaData, FName KeyName, FString Value)
+{
+	Value.TrimStartAndEndInline();
+	Value.TrimQuotesInline();
+
+	FString* ExistingValue = MetaData.Find(KeyName);
+	if (ExistingValue && Value != *ExistingValue)
+	{
+		FError::Throwf(TEXT("Metadata key '%s' first seen with value '%s' then '%s'"), *KeyName.ToString(), **ExistingValue, *Value);
+	}
+
+	// finally we have enough to put it into our metadata
+	MetaData.Add(KeyName, MoveTemp(Value));
+}
+
+
diff --git a/Engine/Source/Programs/UHTLite/Private/BaseParser.h b/Engine/Source/Programs/UHTLite/Private/BaseParser.h
new file mode 100644
--- /dev/null
+++ b/Engine/Source/Programs/UHTLite/Private/BaseParser.h
@@ -0,0 +1,190 @@
+// Copyright Epic Games, Inc. All Rights Reserved.
+
+#pragma once
+
+#include "CoreMinimal.h"
+#include "Containers/UnrealString.h"
+
+class FToken;
+
+/////////////////////////////////////////////////////
+// FBaseParser
+
+
+enum class ESymbolParseOption
+{
+	Normal,
+	CloseTemplateBracket
+};
+
+// A specifier with optional value
+struct FPropertySpecifier
+{
+public:
+	explicit FPropertySpecifier(FString&& InKey)
+		: Key(MoveTemp(InKey))
+	{
+	}
+
+	explicit FPropertySpecifier(const FString& InKey)
+		: Key(InKey)
+	{
+	}
+
+	FString Key;
+	TArray<FString> Values;
+
+	FString ConvertToString() const;
+};
+
+//
+// Base class of header parsers.
+//
+
+class FBaseParser
+{
+protected:
+	FBaseParser();
+
+public:
+	// Input text.
+	const TCHAR* Input;
+
+	// Length of input text.
+	int32 InputLen;
+
+	// Current position in text.
+	int32 InputPos;
+
+	// Current line in text.
+	int32 InputLine;
+
+	// Position previous to last GetChar() call.
+	int32 PrevPos;
+
+	// Line previous to last GetChar() call.
+	int32 PrevLine;
+
+	// Previous comment parsed by GetChar() call.
+	FString PrevComment;
+
+	// Number of statements parsed.
+	int32 StatementsParsed;
+
+	// Total number of lines parsed.
+	int32 LinesParsed;
+
+	void ResetParser(const TCHAR* SourceBuffer, int32 StartingLineNumber = 1);
+
+	// Low-level parsing functions.
+	TCHAR GetChar( bool Literal = false );
+	TCHAR PeekChar();
+	TCHAR GetLeadingChar();
+	void UngetChar();
+
+	/**
+	 * Tests if a character is an end-of-line character.
+	 *
+	 * @param	c	The character to test.
+	 *
+	 * @return	true if c is an end-of-line character, false otherwise.
+	 */
+	static bool IsEOL( TCHAR c );
+
+	/**
+	 * Tests if a character is a whitespace character.
+	 *
+	 * @param	c	The character to test.
+	 *
+	 * @return	true if c is an whitespace character, false otherwise.
+	 */
+	static bool IsWhitespace( TCHAR c );
+
+	/**
+	 * Gets the next token from the input stream, advancing the variables which keep track of the current input position and line.
+	 *
+	 * @param	Token						receives the value of the parsed text; if Token is pre-initialized, special logic is performed
+	 *										to attempt to evaluated Token in the context of that type.  Useful for distinguishing between ambigous symbols
+	 *										like enum tags.
+	 * @param	NoConsts					specify true to indicate that tokens representing literal const values are not allowed.
+	 * @param	ParseTemplateClosingBracket	specify true to treat >> as two template closing brackets instead of shift operator.
+	 *
+	 * @return	true if a token was successfully processed, false otherwise.
+	 */
+	bool GetToken( FToken& Token, bool bNoConsts = false, ESymbolParseOption bParseTemplateClosingBracket = ESymbolParseOption::Normal );
+
+	/**
+	 * Put all text from the current position up to either EOL or the StopToken
+	 * into Token.  Advances the compiler's current position.
+	 *
+	 * @param	Token	[out] will contain the text that was parsed
+	 * @param	StopChar	stop processing when this character is reached
+	 *
+	 * @return	true if a token was parsed
+	 */
+	bool GetRawToken( FToken& Token, TCHAR StopChar = TCHAR('\n') );
+
+	// Doesn't quit if StopChar is found inside a double-quoted string, but does not support quote escapes
+	bool GetRawTokenRespectingQuotes( FToken& Token, TCHAR StopChar = TCHAR('\n') );
+
+	void UngetToken( const FToken& Token );
+	bool GetIdentifier( FToken& Token, bool bNoConsts = false );
+	bool GetSymbol( FToken& Token );
+
+	/**
+	 * Get an int constant
+	 * @return true on success, otherwise false.
+	 */
+	bool GetConstInt(int32& Result, const TCHAR* Tag = NULL);
+	bool GetConstInt64(int64& Result, const TCHAR* Tag = NULL);
+
+	// Matching predefined text.
+	bool MatchIdentifierByName( FName Match );
+	bool MatchIdentifier( const TCHAR* Match, ESearchCase::Type SearchCase);
+	bool MatchConstInt( const TCHAR* Match );
+	bool MatchAnyConstInt();
+	bool PeekIdentifierByName( FName Match );
+	bool PeekIdentifier( const TCHAR* Match, ESearchCase::Type SearchCase);
+	bool MatchSymbol( const TCHAR Match, ESymbolParseOption bParseTemplateClosingBracket = ESymbolParseOption::Normal );
+	bool MatchSymbol(const TCHAR* Match, ESymbolParseOption bParseTemplateClosingBracket = ESymbolParseOption::Normal);
+	void MatchSemi();
+	bool PeekSymbol( const TCHAR Match );
+
+	// Requiring predefined text.
+	void RequireIdentifier( const TCHAR* Match, ESearchCase::Type SearchCase, const TCHAR* Tag );
+	void RequireSymbol( const TCHAR Match, const TCHAR* Tag, ESymbolParseOption bParseTemplateClosingBracket = ESymbolParseOption::Normal );
+	void RequireSymbol(const TCHAR Match, TFunctionRef<FString()> TagGetter, ESymbolParseOption bParseTemplateClosingBracket = ESymbolParseOption::Normal);
+	void RequireConstInt( const TCHAR* Match, const TCHAR* Tag );
+	void RequireAnyConstInt( const TCHAR* Tag );
+
+	/** Clears out the stored comment. */
+	void ClearComment();
+
+	// Reads a new-style value
+	//@TODO: UCREMOVAL: Needs a better name
+	FString ReadNewStyleValue(const TCHAR* TypeOfSpecifier);
+
+	// Reads ['(' Value [',' Value]* ')'] and places each value into the Items array
+	bool ReadOptionalCommaSeparatedListInParens(TArray<FString>& Items, const TCHAR* TypeOfSpecifier);
+
+	//////////////
+	// Complicated* parsing code that needs to be shared between the preparser and the parser
+	// (* i.e., doesn't really belong in the base parser)
+
+	// Expecting Name | (MODULE_API Name)
+	//  Places Name into DeclaredName
+	//  Places MODULE_API (where MODULE varies) into RequiredAPIMacroIfPresent
+	// FailureMessage is printed out if the expectation is broken.
+	void ParseNameWithPotentialAPIMacroPrefix(FString& DeclaredName, FString& RequiredAPIMacroIfPresent, const TCHAR* FailureMessage);
+
+	// Reads a set of specifiers (with optional values) inside the () of a new-style metadata macro like UPROPERTY or UFUNCTION
+	void ReadSpecifierSetInsideMacro(TArray<FPropertySpecifier>& SpecifiersFound, const TCHAR* TypeOfSpecifier, TMap<FName, FString>& MetaData);
+
+	// Validates and inserts one key-value pair into the meta data map
+	static void InsertMetaDataPair(TMap<FName, FString>& MetaData, FString InKey, FString InValue);
+
+	// Validates and inserts one key-value pair into the meta data map
+	static void InsertMetaDataPair(TMap<FName, FString>& MetaData, FName InKey, FString InValue);
+
+	//////////////
+};
diff --git a/Engine/Source/Programs/UHTLite/Private/ClassMaps.cpp b/Engine/Source/Programs/UHTLite/Private/ClassMaps.cpp
new file mode 100644
--- /dev/null
+++ b/Engine/Source/Programs/UHTLite/Private/ClassMaps.cpp
@@ -0,0 +1,48 @@
+// Copyright Epic Games, Inc. All Rights Reserved.
+
+#include "ClassMaps.h"
+#include "UnrealHeaderTool.h"
+#include "UnrealTypeDefinitionInfo.h"
+
+FUnrealSourceFiles GUnrealSourceFilesMap;
+FTypeDefinitionInfoMap GTypeDefinitionInfoMap;
+TMap<const UPackage*, TArray<UField*>> GPackageSingletons;
+FCriticalSection GPackageSingletonsCriticalSection;
+FPublicSourceFileSet GPublicSourceFileSet;
+TMap<FProperty*, FString> GArrayDimensions;
+TMap<UPackage*,  const FManifestModule*> GPackageToManifestModuleMap;
+TMap<void*, uint32> GGeneratedCodeHashes;
+FRWLock GGeneratedCodeHashesLock;
+TMap<UEnum*,  EUnderlyingEnumType> GEnumUnderlyingTypes;
+FClassDeclarations GClassDeclarations;
+TSet<FProperty*> GUnsizedProperties;
+TSet<UField*> GEditorOnlyDataTypes;
+TMap<UStruct*, TTuple<TSharedRef<FUnrealSourceFile>, int32>> GStructToSourceLine;
+TMap<UClass*, FArchiveTypeDefinePair> GClassSerializerMap;
+TSet<FProperty*> GPropertyUsesMemoryImageAllocator;
+
+void FClassDeclarations::AddIfMissing(FName Name, TUniqueFunction<TSharedRef<FClassDeclarationMetaData>()>&& DeclConstructFunc)
+{
+	FRWScopeLock Lock(ClassDeclLock, SLT_Write);
+	if (!ClassDeclarations.Contains(Name))
+	{
+		TSharedRef<FClassDeclarationMetaData> ClassDecl = DeclConstructFunc();
+		ClassDeclarations.Add(Name, MoveTemp(ClassDecl));
+	}
+}
+
+FClassDeclarationMetaData* FClassDeclarations::Find(FName Name)
+{
+	FRWScopeLock Lock(ClassDeclLock, SLT_ReadOnly);
+	if (TSharedRef<FClassDeclarationMetaData>* ClassDecl = ClassDeclarations.Find(Name))
+	{
+		return &ClassDecl->Get();
+	}
+	return nullptr;
+}
+
+FClassDeclarationMetaData& FClassDeclarations::FindChecked(FName Name)
+{
+	FRWScopeLock Lock(ClassDeclLock, SLT_ReadOnly);
+	return ClassDeclarations.FindChecked(Name).Get();
+}
diff --git a/Engine/Source/Programs/UHTLite/Private/ClassMaps.h b/Engine/Source/Programs/UHTLite/Private/ClassMaps.h
new file mode 100644
--- /dev/null
+++ b/Engine/Source/Programs/UHTLite/Private/ClassMaps.h
@@ -0,0 +1,153 @@
+// Copyright Epic Games, Inc. All Rights Reserved.
+
+#pragma once
+
+#include "CoreMinimal.h"
+#include "UObject/Stack.h"
+#include "UObject/ErrorException.h"
+#include "UnderlyingEnumType.h"
+
+#include "UnrealSourceFile.h"
+
+class UField;
+class UClass;
+class FProperty;
+class UPackage;
+class UEnum;
+class FClassDeclarationMetaData;
+class FArchive;
+struct FManifestModule;
+class FUnrealSourceFile;
+class FUnrealTypeDefinitionInfo;
+
+enum class ESerializerArchiveType
+{
+	None,
+	Archive,
+	StructuredArchiveRecord
+};
+
+struct FArchiveTypeDefinePair
+{
+	ESerializerArchiveType ArchiveType;
+	FString EnclosingDefine;
+};
+
+// Wrapper class around TypeDefinition map so we can maintain a parallel by name map
+struct FTypeDefinitionInfoMap
+{
+	void Add(UField* Field, TSharedRef<FUnrealTypeDefinitionInfo>&& Definition)
+	{
+		DefinitionsByField.Add(Field, Definition);
+		DefinitionsByName.Add(Field->GetFName(), MoveTemp(Definition));
+	}
+	bool Contains(const UField* Field) { return DefinitionsByField.Contains(Field); }
+	TSharedRef<FUnrealTypeDefinitionInfo>* Find(const UField* Field) { return DefinitionsByField.Find(Field); }
+	TSharedRef<FUnrealTypeDefinitionInfo>* FindByName(const FName Name) { return DefinitionsByName.Find(Name); }
+	TSharedRef<FUnrealTypeDefinitionInfo>& operator[](const UField* Field) { return DefinitionsByField[Field]; }
+
+private:
+
+	TMap<UField*, TSharedRef<FUnrealTypeDefinitionInfo>> DefinitionsByField;
+	TMap<FName, TSharedRef<FUnrealTypeDefinitionInfo>> DefinitionsByName;
+};
+
+// Wrapper class around ClassDeclarations map so we can control access in a threadsafe manner
+struct FClassDeclarations
+{
+	void AddIfMissing(FName Name, TUniqueFunction<TSharedRef<FClassDeclarationMetaData>()>&& DeclConstructFunc);
+	FClassDeclarationMetaData* Find(FName Name);
+	FClassDeclarationMetaData& FindChecked(FName Name);
+
+private:
+	TMap<FName, TSharedRef<FClassDeclarationMetaData>> ClassDeclarations;
+
+	FRWLock ClassDeclLock;
+};
+
+// Wrapper class around SourceFiles map so we can quickly get a list of source files for a given package
+struct FUnrealSourceFiles
+{
+	void Add(FString&& Filename, TSharedRef<FUnrealSourceFile> SourceFile)
+	{
+		TSharedRef<FUnrealSourceFile>& Value = SourceFilesByString.FindOrAdd(MoveTemp(Filename), SourceFile);
+		if (Value != SourceFile)
+		{
+			FError::Throwf(TEXT("Duplicate filename found with different path '%s'."), *Value.Get().GetFilename());
+		}
+		SourceFilesByPackage.FindOrAdd(SourceFile->GetPackage()).Add(&SourceFile.Get());
+	}
+	const TSharedRef<FUnrealSourceFile>* Find(const FString& Id) const { return SourceFilesByString.Find(Id); }
+	const TArray<FUnrealSourceFile*>* FindFilesForPackage(const UPackage* Package) const { return SourceFilesByPackage.Find(Package); }
+
+private:
+	// A map of all source files indexed by string.
+	TMap<FString, TSharedRef<FUnrealSourceFile>> SourceFilesByString;
+
+	// The list of source files per package. Stored as raw pointer since SourceFilesByString holds shared ref.
+	TMap<UPackage*, TArray<FUnrealSourceFile*>> SourceFilesByPackage;
+};
+
+// Wrapper class around PublicSourceFile set so we can quickly get a list of source files for a given package
+struct FPublicSourceFileSet
+{
+	void Add(FUnrealSourceFile* SourceFile)
+	{
+		SourceFileSet.Add(SourceFile);
+		SourceFilesByPackage.FindOrAdd(SourceFile->GetPackage()).Add(SourceFile);
+	}
+	bool Contains(FUnrealSourceFile* SourceFile) const { return SourceFileSet.Contains(SourceFile); }
+	const TArray<FUnrealSourceFile*>* FindFilesForPackage(const UPackage* Package) const { return SourceFilesByPackage.Find(Package); }
+
+private:
+
+	// The set of all public source files. Stored as raw pointer since FUnrealSourceFiles::SourceFilesByString holds shared ref.
+	TSet<FUnrealSourceFile*> SourceFileSet;
+
+	// The list of public source files per package. Stored as raw pointer since FUnrealSourceFiles::SourceFilesByString holds shared ref.
+	TMap<UPackage*, TArray<FUnrealSourceFile*>> SourceFilesByPackage;
+};
+
+extern FUnrealSourceFiles GUnrealSourceFilesMap;
+extern FTypeDefinitionInfoMap GTypeDefinitionInfoMap;
+extern TMap<const UPackage*, TArray<UField*>> GPackageSingletons;
+extern FCriticalSection GPackageSingletonsCriticalSection;
+extern FPublicSourceFileSet GPublicSourceFileSet;
+extern TMap<FProperty*, FString> GArrayDimensions;
+extern TMap<UPackage*,  const FManifestModule*> GPackageToManifestModuleMap;
+extern TMap<void*, uint32> GGeneratedCodeHashes;
+extern FRWLock GGeneratedCodeHashesLock;
+extern TMap<UEnum*, EUnderlyingEnumType> GEnumUnderlyingTypes;
+extern FClassDeclarations GClassDeclarations;
+extern TSet<FProperty*> GUnsizedProperties;
+extern TSet<UField*> GEditorOnlyDataTypes;
+extern TMap<UStruct*, TTuple<TSharedRef<FUnrealSourceFile>, int32>> GStructToSourceLine;
+extern TMap<UClass*, FArchiveTypeDefinePair> GClassSerializerMap;
+extern TSet<FProperty*> GPropertyUsesMemoryImageAllocator;
+
+/** Types access specifiers. */
+enum EAccessSpecifier
+{
+	ACCESS_NotAnAccessSpecifier = 0,
+	ACCESS_Public,
+	ACCESS_Private,
+	ACCESS_Protected,
+	ACCESS_Num,
+};
+
+inline FArchive& operator<<(FArchive& Ar, EAccessSpecifier& ObjectType)
+{
+	if (Ar.IsLoading())
+	{
+		int32 Value;
+		Ar << Value;
+		ObjectType = EAccessSpecifier(Value);
+	}
+	else if (Ar.IsSaving())
+	{
+		int32 Value = (int32)ObjectType;
+		Ar << Value;
+	}
+
+	return Ar;
+}
