Branching UHT to UHTLite:

diff --git a/Engine/Source/Programs/UHTLite/Private/CodeGenerator.cpp b/Engine/Source/Programs/UHTLite/Private/CodeGenerator.cpp
new file mode 100644
--- /dev/null
+++ b/Engine/Source/Programs/UHTLite/Private/CodeGenerator.cpp
@@ -0,0 +1,7176 @@
+// Copyright Epic Games, Inc. All Rights Reserved.
+
+#include "CoreMinimal.h"
+#include "UnrealHeaderTool.h"
+#include "Misc/AssertionMacros.h"
+#include "HAL/PlatformProcess.h"
+#include "Templates/UnrealTemplate.h"
+#include "Math/UnrealMathUtility.h"
+#include "Containers/UnrealString.h"
+#include "UObject/NameTypes.h"
+#include "Logging/LogMacros.h"
+#include "CoreGlobals.h"
+#include "HAL/FileManager.h"
+#include "Misc/Parse.h"
+#include "Misc/CoreMisc.h"
+#include "Misc/CommandLine.h"
+#include "Misc/FileHelper.h"
+#include "Misc/Paths.h"
+#include "Delegates/Delegate.h"
+#include "Misc/Guid.h"
+#include "Misc/ConfigCacheIni.h"
+#include "Misc/FeedbackContext.h"
+#include "Misc/OutputDeviceNull.h"
+#include "UObject/ErrorException.h"
+#include "UObject/Script.h"
+#include "UObject/ObjectMacros.h"
+#include "UObject/UObjectGlobals.h"
+#include "UObject/Class.h"
+#include "UObject/Package.h"
+#include "UObject/MetaData.h"
+#include "UObject/Interface.h"
+#include "UObject/UnrealType.h"
+#include "UObject/TextProperty.h"
+#include "UObject/FieldPathProperty.h"
+#include "Misc/PackageName.h"
+#include "UnrealHeaderToolGlobals.h"
+
+#include "ParserClass.h"
+#include "Scope.h"
+#include "HeaderProvider.h"
+#include "GeneratedCodeVersion.h"
+#include "SimplifiedParsingClassInfo.h"
+#include "UnrealSourceFile.h"
+#include "ParserHelper.h"
+#include "Classes.h"
+#include "NativeClassExporter.h"
+#include "ProfilingDebugging/ScopedTimers.h"
+#include "HeaderParser.h"
+#include "IScriptGeneratorPluginInterface.h"
+#include "Manifest.h"
+#include "StringUtils.h"
+#include "Features/IModularFeatures.h"
+#include "Algo/Copy.h"
+#include "Algo/Sort.h"
+#include "Algo/Reverse.h"
+#include "Async/Async.h"
+#include "Async/ParallelFor.h"
+#include "Misc/ScopeExit.h"
+#include "UnrealTypeDefinitionInfo.h"
+
+#include "FileLineException.h"
+#include "UObject/FieldIterator.h"
+#include "UObject/FieldPath.h"
+
+#include "UObject/WeakFieldPtr.h"
+#include "Templates/SubclassOf.h"
+
+/////////////////////////////////////////////////////
+// Globals
+
+FManifest GManifest;
+
+double GMacroizeTime = 0.0;
+
+static TArray<FString> ChangeMessages;
+static bool bWriteContents = false;
+static bool bVerifyContents = false;
+
+struct FPerHeaderData
+{
+	TSharedPtr<FUnrealSourceFile> UnrealSourceFile;
+	TArray<FHeaderProvider> DependsOn;
+	TArray<FSimplifiedParsingClassInfo> ParsedClassArray;
+};
+
+static void PerformSimplifiedClassParse(UPackage* InParent, const TCHAR* FileName, const TCHAR* Buffer, FPerHeaderData& PerHeaderData);
+static void ProcessInitialClassParse(FPerHeaderData& PerHeaderData);
+
+FCompilerMetadataManager GScriptHelper;
+
+// Array of all the temporary header async file tasks so we can ensure they have completed before issuing our timings
+static FGraphEventArray GAsyncFileTasks;
+
+bool HasIdentifierExactMatch(const TCHAR* StringBegin, const TCHAR* StringEnd, const FString& Find);
+
+namespace
+{
+	static const FName NAME_SerializeToFArchive("SerializeToFArchive");
+	static const FName NAME_SerializeToFStructuredArchive("SerializeToFStructuredArchive");
+	static const FName NAME_ObjectInitializerConstructorDeclared("ObjectInitializerConstructorDeclared");
+	static const FName NAME_InitializeStaticSearchableValues("InitializeStaticSearchableValues");
+	static const FName NAME_OverrideNativeName("OverrideNativeName");
+	static const FName NAME_NoGetter("NoGetter");
+	static const FName NAME_GetByRef("GetByRef");
+
+	static const FString STRING_StructPackage(TEXT("StructPackage"));
+
+	static const int32 HEADER_GENERATED_DELEGATE_SIGNATURE_SUFFIX_LENGTH = FString(HEADER_GENERATED_DELEGATE_SIGNATURE_SUFFIX).Len(); 
+
+	static FString AsTEXT(const FString& InStr)
+	{
+		return FString::Printf(TEXT("TEXT(\"%s\")"), *InStr);
+	}
+
+	const TCHAR HeaderCopyright[] =
+		TEXT("// Copyright Epic Games, Inc. All Rights Reserved.\r\n")
+		TEXT("/*===========================================================================\r\n")
+		TEXT("\tGenerated code exported from UnrealHeaderTool.\r\n")
+		TEXT("\tDO NOT modify this manually! Edit the corresponding .h files instead!\r\n")
+		TEXT("===========================================================================*/\r\n")
+		LINE_TERMINATOR;
+
+	const TCHAR RequiredCPPIncludes[] = TEXT("#include \"UObject/GeneratedCppIncludes.h\"") LINE_TERMINATOR;
+
+	const TCHAR EnableDeprecationWarnings[] = TEXT("PRAGMA_ENABLE_DEPRECATION_WARNINGS") LINE_TERMINATOR;
+	const TCHAR DisableDeprecationWarnings[] = TEXT("PRAGMA_DISABLE_DEPRECATION_WARNINGS") LINE_TERMINATOR;
+
+	// A struct which emits #if and #endif blocks as appropriate when invoked.
+	struct FMacroBlockEmitter
+	{
+		explicit FMacroBlockEmitter(FOutputDevice& InOutput, const TCHAR* InMacro)
+			: Output(InOutput)
+			, bEmittedIf(false)
+			, Macro(InMacro)
+		{
+		}
+
+		~FMacroBlockEmitter()
+		{
+			if (bEmittedIf)
+			{
+				Output.Logf(TEXT("#endif // %s\r\n"), Macro);
+			}
+		}
+
+		void operator()(bool bInBlock)
+		{
+			if (!bEmittedIf && bInBlock)
+			{
+				Output.Logf(TEXT("#if %s\r\n"), Macro);
+				bEmittedIf = true;
+			}
+			else if (bEmittedIf && !bInBlock)
+			{
+				Output.Logf(TEXT("#endif // %s\r\n"), Macro);
+				bEmittedIf = false;
+			}
+		}
+
+		FMacroBlockEmitter(const FMacroBlockEmitter&) = delete;
+		FMacroBlockEmitter& operator=(const FMacroBlockEmitter&) = delete;
+
+	private:
+		FOutputDevice& Output;
+		bool bEmittedIf;
+		const TCHAR* Macro;
+	};
+
+	/** Guard that should be put at the start editor only generated code */
+	const TCHAR BeginEditorOnlyGuard[] = TEXT("#if WITH_EDITOR") LINE_TERMINATOR;
+
+	/** Guard that should be put at the end of editor only generated code */
+	const TCHAR EndEditorOnlyGuard[] = TEXT("#endif //WITH_EDITOR") LINE_TERMINATOR;
+
+	/** Whether or not the given class has any replicated properties. */
+	static bool ClassHasReplicatedProperties(UClass* Class)
+	{
+		if (!Class->HasAnyClassFlags(CLASS_ReplicationDataIsSetUp))
+		{
+			for (TFieldIterator<FProperty> It(Class, EFieldIteratorFlags::ExcludeSuper); It; ++It)
+		{
+			if ((It->PropertyFlags & CPF_Net) != 0)
+			{
+				return true;
+			}
+		}
+		}
+
+		return Class->FirstOwnedClassRep < Class->ClassReps.Num();
+	}
+
+	static void ExportNetData(FOutputDevice& Out, UClass* Class, const TCHAR* API)
+	{
+		const TArray<FRepRecord>& ClassReps = Class->ClassReps;
+
+		FUHTStringBuilder NetFieldBuilder;
+		NetFieldBuilder.Logf(TEXT(""
+		"\tenum class ENetFields_Private : uint16\r\n"
+		"\t{\r\n"
+		"\t\tNETFIELD_REP_START=(uint16)((int32)Super::ENetFields_Private::NETFIELD_REP_END + (int32)1),\r\n"));
+
+		FUHTStringBuilder ArrayDimBuilder;
+
+		bool bAnyStaticArrays = false;
+		bool bIsFirst = true;
+		for (int32 ClassRepIndex = Class->FirstOwnedClassRep; ClassRepIndex < ClassReps.Num(); ++ClassRepIndex)
+		{
+			const FRepRecord& ClassRep = ClassReps[ClassRepIndex];
+			const FString PropertyName = ClassRep.Property->GetName();
+
+			if (ClassRep.Property->ArrayDim == 1)
+			{
+				if (UNLIKELY(bIsFirst))
+				{
+					NetFieldBuilder.Logf(TEXT("\t\t%s=NETFIELD_REP_START,\r\n"), *PropertyName);
+					bIsFirst = false;
+				}
+				else
+				{
+					NetFieldBuilder.Logf(TEXT("\t\t%s,\r\n"), *PropertyName);
+				}
+			}
+			else
+			{
+				bAnyStaticArrays = true;
+				ArrayDimBuilder.Logf(TEXT("\t\t%s=%s,\r\n"), *PropertyName, *GArrayDimensions.FindChecked(ClassReps[ClassRepIndex].Property));
+
+				if (UNLIKELY(bIsFirst))
+				{
+					NetFieldBuilder.Logf(TEXT("\t\t%s_STATIC_ARRAY=NETFIELD_REP_START,\r\n"), *PropertyName);
+					bIsFirst = false;
+				}
+				else
+				{
+					NetFieldBuilder.Logf(TEXT("\t\t%s_STATIC_ARRAY,\r\n"), *PropertyName);
+				}
+
+				NetFieldBuilder.Logf(TEXT("\t\t%s_STATIC_ARRAY_END=((uint16)%s_STATIC_ARRAY + (uint16)EArrayDims_Private::%s - (uint16)1),\r\n"), *PropertyName, *PropertyName, *PropertyName);
+			}
+		}
+
+		const FProperty* LastProperty = ClassReps.Last().Property;
+		NetFieldBuilder.Logf(TEXT("\t\tNETFIELD_REP_END=%s%s"), *LastProperty->GetName(), LastProperty->ArrayDim > 1 ? TEXT("_STATIC_ARRAY_END") : TEXT(""));
+
+		NetFieldBuilder.Log(TEXT("\t};"));
+
+		if (bAnyStaticArrays)
+		{
+			Out.Logf(TEXT(""
+				"\tenum class EArrayDims_Private : uint16\r\n"
+				"\t{\r\n"
+				"%s"
+				"\t};\r\n"), *ArrayDimBuilder);
+		}
+
+		Out.Logf(TEXT(""
+			"%s\r\n" // NetFields
+			"\t%s_API virtual void ValidateGeneratedRepEnums(const TArray<struct FRepRecord>& ClassReps) const override;\r\n"),
+			*NetFieldBuilder,
+			API);
+	}
+
+	static const FString STRING_GetLifetimeReplicatedPropsStr(TEXT("GetLifetimeReplicatedProps"));
+
+	static void WriteReplicatedMacroData(
+		const ClassDefinitionRange& ClassRange,
+		const TCHAR* ClassCPPName,
+		const TCHAR* API,
+		FClass* Class,
+		FClass* SuperClass,
+		FOutputDevice& Writer,
+		const FUnrealSourceFile& SourceFile,
+		FNativeClassHeaderGenerator::EExportClassOutFlags& OutFlags)
+	{
+		const bool bHasGetLifetimeReplicatedProps = HasIdentifierExactMatch(ClassRange.Start, ClassRange.End, STRING_GetLifetimeReplicatedPropsStr);
+
+		if (!bHasGetLifetimeReplicatedProps)
+		{
+			// Default version autogenerates declarations.
+			if (SourceFile.GetGeneratedCodeVersionForStruct(Class) == EGeneratedCodeVersion::V1)
+			{
+				Writer.Logf(TEXT("\tvoid GetLifetimeReplicatedProps(TArray<FLifetimeProperty>& OutLifetimeProps) const override;\r\n"));
+			}
+			else
+			{
+				FError::Throwf(TEXT("Class %s has Net flagged properties and should declare member function: void GetLifetimeReplicatedProps(TArray<FLifetimeProperty>& OutLifetimeProps) const override"), ClassCPPName);
+			}
+		}
+
+
+		ExportNetData(Writer, Class, API);
+		
+		// If this class has replicated properties and it owns the first one, that means
+		// it's the base most replicated class. In that case, go ahead and add our interface macro.
+		if (Class->ClassReps.Num() > 0 && Class->FirstOwnedClassRep == 0)
+		{
+			OutFlags |= FNativeClassHeaderGenerator::EExportClassOutFlags::NeedsPushModelHeaders;
+			Writer.Logf(TEXT(
+				"private:\r\n"
+				"\tREPLICATED_BASE_CLASS(%s%s)\r\n"
+				"public:\r\n"
+			), Class->GetPrefixCPP(), *Class->GetName());
+		}
+	}
+}
+
+#define BEGIN_WRAP_EDITOR_ONLY(DoWrap) DoWrap ? BeginEditorOnlyGuard : TEXT("")
+#define END_WRAP_EDITOR_ONLY(DoWrap) DoWrap ? EndEditorOnlyGuard : TEXT("")
+
+/**
+ * Finds exact match of Identifier in string. Returns nullptr if none is found.
+ *
+ * @param StringBegin Start of string to search.
+ * @param StringEnd End of string to search.
+ * @param Identifier Identifier to find.
+ * @return Pointer to Identifier match within string. nullptr if none found.
+ */
+const TCHAR* FindIdentifierExactMatch(const TCHAR* StringBegin, const TCHAR* StringEnd, const FString& Identifier)
+{
+	int32 StringLen = UE_PTRDIFF_TO_INT32(StringEnd - StringBegin);
+
+	// Check for exact match first.
+	if (FCString::Strncmp(StringBegin, *Identifier, StringLen) == 0)
+	{
+		return StringBegin;
+	}
+
+	int32        FindLen        = Identifier.Len();
+	const TCHAR* StringToSearch = StringBegin;
+
+	for (;;)
+	{
+		const TCHAR* IdentifierStart = FCString::Strstr(StringToSearch, *Identifier);
+		if (IdentifierStart == nullptr)
+		{
+			// Not found.
+			return nullptr;
+		}
+
+		if (IdentifierStart > StringEnd || IdentifierStart + FindLen + 1 > StringEnd)
+		{
+			// Found match is out of string range.
+			return nullptr;
+		}
+
+		if (IdentifierStart == StringBegin && !FChar::IsIdentifier(*(IdentifierStart + FindLen + 1)))
+		{
+			// Found match is at the beginning of string.
+			return IdentifierStart;
+		}
+
+		if (IdentifierStart + FindLen == StringEnd && !FChar::IsIdentifier(*(IdentifierStart - 1)))
+		{
+			// Found match ends with end of string.
+			return IdentifierStart;
+		}
+
+		if (!FChar::IsIdentifier(*(IdentifierStart + FindLen)) && !FChar::IsIdentifier(*(IdentifierStart - 1)))
+		{
+			// Found match is in the middle of string
+			return IdentifierStart;
+		}
+
+		// Didn't find exact match, nor got to end of search string. Keep on searching.
+		StringToSearch = IdentifierStart + FindLen;
+	}
+
+	// We should never get here.
+	checkNoEntry();
+	return nullptr;
+}
+
+/**
+ * Finds exact match of Identifier in string. Returns nullptr if none is found.
+ *
+ * @param String String to search.
+ * @param Identifier Identifier to find.
+ * @return Index to Identifier match within String. INDEX_NONE if none found.
+ */
+int32 FindIdentifierExactMatch(const FString& String, const FString& Identifier)
+{
+	const TCHAR* IdentifierPtr = FindIdentifierExactMatch(*String, *String + String.Len(), Identifier);
+	if (IdentifierPtr == nullptr)
+	{
+		return INDEX_NONE;
+	}
+
+	return UE_PTRDIFF_TO_INT32(IdentifierPtr - *String);
+}
+
+/**
+* Checks if exact match of Identifier is in String.
+*
+* @param StringBegin Start of string to search.
+* @param StringEnd End of string to search.
+* @param Identifier Identifier to find.
+* @return true if Identifier is within string, false otherwise.
+*/
+bool HasIdentifierExactMatch(const TCHAR* StringBegin, const TCHAR* StringEnd, const FString& Find)
+{
+	return FindIdentifierExactMatch(StringBegin, StringEnd, Find) != nullptr;
+}
+
+/**
+* Checks if exact match of Identifier is in String.
+*
+* @param String String to search.
+* @param Identifier Identifier to find.
+* @return true if Identifier is within String, false otherwise.
+*/
+bool HasIdentifierExactMatch(const FString &String, const FString& Identifier)
+{
+	return FindIdentifierExactMatch(String, Identifier) != INDEX_NONE;
+}
+
+void ConvertToBuildIncludePath(const UPackage* Package, FString& LocalPath)
+{
+	FPaths::MakePathRelativeTo(LocalPath, *GPackageToManifestModuleMap.FindChecked(Package)->IncludeBase);
+}
+
+/**
+ *	Helper function to retrieve the package manifest
+ *
+ *	@param	InPackage		The name of the package of interest
+ *
+ *	@return	The manifest if found
+ */
+FManifestModule* GetPackageManifest(const FString& CheckPackage)
+{
+	// Mapping of processed packages to their locations
+	// An empty location string means it was processed but not found
+	static TMap<FString, FManifestModule*> CheckedPackageList;
+
+	FManifestModule* ModuleInfoPtr = CheckedPackageList.FindRef(CheckPackage);
+
+	if (!ModuleInfoPtr)
+	{
+		FManifestModule* ModuleInfoPtr2 = GManifest.Modules.FindByPredicate([&](FManifestModule& Module) { return Module.Name == CheckPackage; });
+		if (ModuleInfoPtr2 && IFileManager::Get().DirectoryExists(*ModuleInfoPtr2->BaseDirectory))
+		{
+			ModuleInfoPtr = ModuleInfoPtr2;
+			CheckedPackageList.Add(CheckPackage, ModuleInfoPtr);
+		}
+	}
+
+	return ModuleInfoPtr;
+}
+
+FString Macroize(const TCHAR* MacroName, FString&& StringToMacroize)
+{
+	FScopedDurationTimer Tracker(GMacroizeTime);
+
+	FString Result(MoveTemp(StringToMacroize));
+	if (Result.Len())
+	{
+		Result.ReplaceInline(TEXT("\r\n"), TEXT("\n"), ESearchCase::CaseSensitive);
+		Result.ReplaceInline(TEXT("\n"), TEXT(" \\\n"), ESearchCase::CaseSensitive);
+		checkSlow(Result.EndsWith(TEXT(" \\\n"), ESearchCase::CaseSensitive));
+
+		if (Result.Len() >= 3)
+		{
+			for (int32 Index = Result.Len() - 3; Index < Result.Len(); ++Index)
+			{
+				Result[Index] = TEXT('\n');
+			}
+		}
+		else
+		{
+			Result = TEXT("\n\n\n");
+		}
+		Result.ReplaceInline(TEXT("\n"), TEXT("\r\n"), ESearchCase::CaseSensitive);
+	}
+	return FString::Printf(TEXT("#define %s%s\r\n%s"), MacroName, Result.Len() ? TEXT(" \\") : TEXT(""), *Result);
+}
+
+static void AddGeneratedCodeHash(void* Field, uint32 Hash)
+{
+	FRWScopeLock Lock(GGeneratedCodeHashesLock, SLT_Write);
+	GGeneratedCodeHashes.Add(Field, Hash);
+}
+
+/** Generates a Hash tag string for the specified field */
+static FString GetGeneratedCodeHashTag(void* Field)
+{
+	FString Tag;
+	bool bFoundHash = false;
+	uint32 Hash = 0;
+
+	{
+		FRWScopeLock Lock(GGeneratedCodeHashesLock, SLT_ReadOnly);
+		if (const uint32* FieldHash = GGeneratedCodeHashes.Find(Field))
+		{
+			bFoundHash = true;
+			Hash = *FieldHash;
+		}
+	}
+
+	if (bFoundHash)
+	{
+		Tag = FString::Printf(TEXT(" // %u"), Hash);
+	}
+	return Tag;
+}
+
+struct FParmsAndReturnProperties
+{
+	FParmsAndReturnProperties()
+		: Return(nullptr)
+	{
+	}
+
+	bool HasParms() const
+	{
+		return Parms.Num() || Return;
+	}
+
+	TArray<FProperty*> Parms;
+	FProperty*         Return;
+};
+
+/**
+ * Get parameters and return type for a given function.
+ *
+ * @param  Function The function to get the parameters for.
+ * @return An aggregate containing the parameters and return type of that function.
+ */
+FParmsAndReturnProperties GetFunctionParmsAndReturn(UFunction* Function)
+{
+	FParmsAndReturnProperties Result;
+	for ( TFieldIterator<FProperty> It(Function); It; ++It)
+	{
+		FProperty* Field = *It;
+
+		if ((It->PropertyFlags & (CPF_Parm | CPF_ReturnParm)) == CPF_Parm)
+		{
+			Result.Parms.Add(Field);
+		}
+		else if (It->PropertyFlags & CPF_ReturnParm)
+		{
+			Result.Return = Field;
+		}
+	}
+	return Result;
+}
+
+/**
+ * Determines whether the glue version of the specified native function
+ * should be exported
+ *
+ * @param	Function	the function to check
+ * @return	true if the glue version of the function should be exported.
+ */
+bool ShouldExportUFunction(UFunction* Function)
+{
+	// export any script stubs for native functions declared in interface classes
+	bool bIsBlueprintNativeEvent = (Function->FunctionFlags & FUNC_BlueprintEvent) && (Function->FunctionFlags & FUNC_Native);
+	if (Function->GetOwnerClass()->HasAnyClassFlags(CLASS_Interface) && !bIsBlueprintNativeEvent)
+	{
+		return true;
+	}
+
+	// always export if the function is static
+	if (Function->FunctionFlags & FUNC_Static)
+	{
+		return true;
+	}
+
+	// don't export the function if this is not the original declaration and there is
+	// at least one parent version of the function that is declared native
+	for (UFunction* ParentFunction = Function->GetSuperFunction(); ParentFunction; ParentFunction = ParentFunction->GetSuperFunction())
+	{
+		if (ParentFunction->FunctionFlags & FUNC_Native)
+		{
+			return false;
+		}
+	}
+
+	return true;
+}
+
+FString CreateLiteralString(const FString& Str)
+{
+	FString Result;
+
+	// Have a reasonable guess at reserving the right size
+	Result.Reserve(Str.Len() + 8);
+	Result += TEXT("TEXT(\"");
+
+	bool bPreviousCharacterWasHex = false;
+
+	const TCHAR* Ptr = *Str;
+	while (TCHAR Ch = *Ptr++)
+	{
+		switch (Ch)
+		{
+			case TEXT('\r'): continue;
+			case TEXT('\n'): Result += TEXT("\\n");  bPreviousCharacterWasHex = false; break;
+			case TEXT('\\'): Result += TEXT("\\\\"); bPreviousCharacterWasHex = false; break;
+			case TEXT('\"'): Result += TEXT("\\\""); bPreviousCharacterWasHex = false; break;
+			default:
+				if (Ch < 31 || Ch >= 128)
+				{
+					Result += FString::Printf(TEXT("\\x%04x"), Ch);
+					bPreviousCharacterWasHex = true;
+				}
+				else
+				{
+					// We close and open the literal (with TEXT) here in order to ensure that successive hex characters aren't appended to the hex sequence, causing a different number
+					if (bPreviousCharacterWasHex && FCharWide::IsHexDigit(Ch))
+					{
+						Result += "\")TEXT(\"";
+					}
+					bPreviousCharacterWasHex = false;
+					Result += Ch;
+				}
+				break;
+		}
+	}
+
+	Result += TEXT("\")");
+	return Result;
+}
+
+FString CreateUTF8LiteralString(const FString& Str)
+{
+	FString Result;
+
+	// Have a reasonable guess at reserving the right size
+	Result.Reserve(Str.Len() + 2);
+	Result += TEXT("\"");
+
+	bool bPreviousCharacterWasHex = false;
+
+	FTCHARToUTF8 StrUTF8(*Str);
+
+	const char* Ptr = StrUTF8.Get();
+	while (char Ch = *Ptr++)
+	{
+		switch (Ch)
+		{
+			case '\r': continue;
+			case '\n': Result += TEXT("\\n");  bPreviousCharacterWasHex = false; break;
+			case '\\': Result += TEXT("\\\\"); bPreviousCharacterWasHex = false; break;
+			case '\"': Result += TEXT("\\\""); bPreviousCharacterWasHex = false; break;
+			default:
+				if (Ch < 31)
+				{
+					Result += FString::Printf(TEXT("\\x%02x"), (uint8)Ch);
+					bPreviousCharacterWasHex = true;
+				}
+				else
+				{
+					// We close and open the literal here in order to ensure that successive hex characters aren't appended to the hex sequence, causing a different number
+					if (bPreviousCharacterWasHex && FCharWide::IsHexDigit(Ch))
+					{
+						Result += "\"\"";
+					}
+					bPreviousCharacterWasHex = false;
+					Result += Ch;
+				}
+				break;
+		}
+	}
+
+	Result += TEXT("\"");
+	return Result;
+}
+
+TMap<FName, FString> GenerateMetadataMapForObject(const UObject* Obj)
+{
+	check(Obj);
+	UPackage* Package = Obj->GetOutermost();
+	check(Package);
+	UMetaData* Metadata = Package->GetMetaData();
+	check(Metadata);
+
+	TMap<FName, FString>* PackageMap = Metadata->ObjectMetaDataMap.Find(Obj);
+	TMap<FName, FString> Map;
+	if (PackageMap)
+	{
+		for (const TPair<FName, FString>& MetaKeyValue : *PackageMap)
+		{
+			FString Key = MetaKeyValue.Key.ToString();
+			if (!Key.StartsWith(TEXT("/Script")))
+			{
+				Map.Add(MetaKeyValue.Key, MetaKeyValue.Value);
+			}
+		}
+	}
+	return Map;
+}
+
+TMap<FName, FString> GenerateMetadataMapForField(const FField* Field)
+{
+	TMap<FName, FString> MetaDataMap;
+	const TMap<FName, FString>* FieldMetaDataMap = Field->GetMetaDataMap();
+	if (FieldMetaDataMap)
+	{
+		MetaDataMap = *FieldMetaDataMap;
+	}
+	return MetaDataMap;
+}
+
+// Returns the METADATA_PARAMS for this output
+static FString OutputMetaDataCodeForObject(FOutputDevice& OutDeclaration, FOutputDevice& Out, FFieldVariant Object, const TCHAR* MetaDataBlockName, const TCHAR* DeclSpaces, const TCHAR* Spaces)
+{
+	TMap<FName, FString> MetaData;
+	
+	if (Object.IsUObject())
+	{
+		MetaData = GenerateMetadataMapForObject(Object.ToUObject());
+	}
+	else
+	{
+		MetaData = GenerateMetadataMapForField(Object.ToField());
+	}
+
+	FString Result;
+	if (MetaData.Num())
+	{
+		typedef TKeyValuePair<FName, FString*> KVPType;
+		TArray<KVPType> KVPs;
+		KVPs.Reserve(MetaData.Num());
+		for (TPair<FName, FString>& KVP : MetaData)
+		{
+			KVPs.Add(KVPType(KVP.Key, &KVP.Value));
+		}
+
+		// We sort the metadata here so that we can get consistent output across multiple runs
+		// even when metadata is added in a different order
+		Algo::SortBy(KVPs, &KVPType::Key, FNameLexicalLess());
+
+		FString MetaDataBlockNameWithoutScope = MetaDataBlockName;
+		int32 ScopeIndex = MetaDataBlockNameWithoutScope.Find(TEXT("::"), ESearchCase::CaseSensitive);
+		if (ScopeIndex != INDEX_NONE)
+		{
+			MetaDataBlockNameWithoutScope.RightChopInline(ScopeIndex + 2, false);
+		}
+
+		OutDeclaration.Log (TEXT("#if WITH_METADATA\r\n"));
+		OutDeclaration.Logf(TEXT("%sstatic const UE4CodeGen_Private::FMetaDataPairParam %s[];\r\n"), DeclSpaces, *MetaDataBlockNameWithoutScope);
+		OutDeclaration.Log (TEXT("#endif\r\n"));
+
+		Out.Log (TEXT("#if WITH_METADATA\r\n"));
+		Out.Logf(TEXT("%sconst UE4CodeGen_Private::FMetaDataPairParam %s[] = {\r\n"), Spaces, MetaDataBlockName);
+
+		for (const KVPType& KVP : KVPs)
+		{
+			Out.Logf(TEXT("%s\t{ %s, %s },\r\n"), Spaces, *CreateUTF8LiteralString(KVP.Key.ToString()), *CreateUTF8LiteralString(*KVP.Value));
+		}
+
+		Out.Logf(TEXT("%s};\r\n"), Spaces);
+		Out.Log (TEXT("#endif\r\n"));
+
+		Result = FString::Printf(TEXT("METADATA_PARAMS(%s, UE_ARRAY_COUNT(%s))"), MetaDataBlockName, MetaDataBlockName);
+	}
+	else
+	{
+		Result = TEXT("METADATA_PARAMS(nullptr, 0)");
+	}
+
+	return Result;
+}
+
+void FNativeClassHeaderGenerator::ExportProperties(FOutputDevice& Out, UStruct* Struct, int32 TextIndent)
+{
+	FProperty*	Previous			= NULL;
+	FProperty*	PreviousNonEditorOnly = NULL;
+	FProperty*	LastInSuper			= NULL;
+	UStruct*	InheritanceSuper	= Struct->GetInheritanceSuper();
+
+	// Find last property in the lowest base class that has any properties
+	UStruct* CurrentSuper = InheritanceSuper;
+	while (LastInSuper == NULL && CurrentSuper)
+	{
+		for( TFieldIterator<FProperty> It(CurrentSuper,EFieldIteratorFlags::ExcludeSuper); It; ++It )
+		{
+			FProperty* Current = *It;
+
+			// Disregard properties with 0 size like functions.
+			if( It.GetStruct() == CurrentSuper && Current->ElementSize )
+			{
+				LastInSuper = Current;
+			}
+		}
+		// go up a layer in the hierarchy
+		CurrentSuper = CurrentSuper->GetSuperStruct();
+	}
+
+	FMacroBlockEmitter WithEditorOnlyData(Out, TEXT("WITH_EDITORONLY_DATA"));
+
+	// Iterate over all properties in this struct.
+	for( TFieldIterator<FProperty> It(Struct, EFieldIteratorFlags::ExcludeSuper); It; ++It )
+	{
+		FProperty* Current = *It;
+
+		// Disregard properties with 0 size like functions.
+		if (It.GetStruct() == Struct)
+		{
+			WithEditorOnlyData(Current->IsEditorOnlyProperty());
+
+			// Export property specifiers
+			// Indent code and export CPP text.
+			{
+				FUHTStringBuilder JustPropertyDecl;
+
+				const FString* Dim = GArrayDimensions.Find(Current);
+				Current->ExportCppDeclaration( JustPropertyDecl, EExportedDeclaration::Member, Dim ? **Dim : NULL);
+				ApplyAlternatePropertyExportText(*It, JustPropertyDecl, EExportingState::TypeEraseDelegates);
+
+				// Finish up line.
+				Out.Logf(TEXT("%s%s;\r\n"), FCString::Tab(TextIndent + 1), *JustPropertyDecl);
+			}
+
+			LastInSuper	= NULL;
+			Previous = Current;
+			if (!Current->IsEditorOnlyProperty())
+			{
+				PreviousNonEditorOnly = Current;
+			}
+		}
+	}
+}
+
+/**
+ * Class that is representing a type singleton.
+ */
+struct FTypeSingleton
+{
+public:
+	/** Constructor */
+	FTypeSingleton(FString InName, UField* InType)
+		: Name(MoveTemp(InName)), Type(InType) {}
+
+	/**
+	 * Gets this singleton's name.
+	 */
+	const FString& GetName() const
+	{
+		return Name;
+	}
+
+	/**
+	 * Gets this singleton's extern declaration.
+	 */
+	const FString& GetExternDecl() const
+	{
+		FRWScopeLock Lock(ExternDeclLock, SLT_ReadOnly);
+		if (ExternDecl.IsEmpty())
+		{
+			Lock.ReleaseReadOnlyLockAndAcquireWriteLock_USE_WITH_CAUTION();
+
+			// Verify the decl is still empty in case another thread had also been waiting on writing this data and got the write lock first
+			if (ExternDecl.IsEmpty())
+			{
+				ExternDecl = GenerateExternDecl(Type, GetName());
+			}
+		}
+
+		return ExternDecl;
+	}
+
+private:
+	/**
+	 * Extern declaration generator.
+	 */
+	static FString GenerateExternDecl(UField* InType, const FString& InName)
+	{
+		const TCHAR* TypeStr = nullptr;
+
+		if (InType->GetClass() == UClass::StaticClass())
+		{
+			TypeStr = TEXT("UClass");
+		}
+		else if (InType->GetClass() == UFunction::StaticClass() || InType->GetClass() == UDelegateFunction::StaticClass() || InType->GetClass() == USparseDelegateFunction::StaticClass())
+		{
+			TypeStr = TEXT("UFunction");
+		}
+		else if (InType->GetClass() == UScriptStruct::StaticClass())
+		{
+			TypeStr = TEXT("UScriptStruct");
+		}
+		else if (InType->GetClass() == UEnum::StaticClass())
+		{
+			TypeStr = TEXT("UEnum");
+		}
+		else
+		{
+			FError::Throwf(TEXT("Unsupported item type to get extern for."));
+		}
+
+		return FString::Printf(
+			TEXT("\t%s_API %s* %s;\r\n"),
+			*FPackageName::GetShortName(InType->GetOutermost()).ToUpper(),
+			TypeStr,
+			*InName
+		);
+	}
+
+	/** Field that stores this singleton name. */
+	FString Name;
+
+	/** Cached field that stores this singleton extern declaration. */
+	mutable FString ExternDecl;
+
+	/** Mutex to ensure 2 threads don't try to generate the extern decl at the same time. */
+	mutable FRWLock ExternDeclLock;
+
+	/** Type of the singleton */
+	UField* Type;
+};
+
+/**
+ * Class that represents type singleton cache.
+ */
+class FTypeSingletonCache
+{
+public:
+	/**
+	 * Gets type singleton from cache.
+	 *
+	 * @param Type Singleton type.
+	 * @param bRequiresValidObject Does it require a valid object?
+	 */
+	static const FTypeSingleton& Get(UField* Type, bool bRequiresValidObject = true)
+	{
+		FTypeSingletonCacheKey Key(Type, bRequiresValidObject);
+
+		FRWScopeLock Lock(Mutex, SLT_ReadOnly);
+		TUniquePtr<FTypeSingleton>* SingletonPtr = CacheData.Find(Key);
+		if (SingletonPtr == nullptr)
+		{
+			TUniquePtr<FTypeSingleton> TypeSingleton(MakeUnique<FTypeSingleton>(GenerateSingletonName(Type, bRequiresValidObject), Type));
+
+			Lock.ReleaseReadOnlyLockAndAcquireWriteLock_USE_WITH_CAUTION();
+
+			// Check the map again in case another thread had also been waiting on writing this data and got the write lock first
+			SingletonPtr = CacheData.Find(Key);
+			if (SingletonPtr == nullptr)
+			{
+				SingletonPtr = &CacheData.Add(Key, MoveTemp(TypeSingleton));
+			}
+		}
+
+		return **SingletonPtr;
+	}
+
+private:
+	/**
+	 * Private type that represents cache map key.
+	 */
+	struct FTypeSingletonCacheKey
+	{
+		/** FTypeSingleton type */
+		UField* Type;
+
+		/** If this type singleton requires valid object. */
+		bool bRequiresValidObject;
+
+		/* Constructor */
+		FTypeSingletonCacheKey(UField* InType, bool bInRequiresValidObject)
+			: Type(InType), bRequiresValidObject(bInRequiresValidObject)
+		{}
+
+		/**
+		 * Equality operator.
+		 *
+		 * @param Other Other key.
+		 *
+		 * @returns True if this is equal to Other. False otherwise.
+		 */
+		bool operator==(const FTypeSingletonCacheKey& Other) const
+		{
+			return Type == Other.Type && bRequiresValidObject == Other.bRequiresValidObject;
+		}
+
+		/**
+		 * Gets hash value for this object.
+		 */
+		friend uint32 GetTypeHash(const FTypeSingletonCacheKey& Object)
+		{
+			return HashCombine(
+				GetTypeHash(Object.Type),
+				GetTypeHash(Object.bRequiresValidObject)
+			);
+		}
+	};
+
+	/**
+	 * Generates singleton name.
+	 */
+	static FString GenerateSingletonName(UField* Item, bool bRequiresValidObject)
+	{
+		check(Item);
+
+		bool bNoRegister = false;
+		if (UClass* ItemClass = Cast<UClass>(Item))
+		{
+			if (!bRequiresValidObject && !ItemClass->HasAllClassFlags(CLASS_Intrinsic))
+			{
+				bNoRegister = true;
+			}
+		}
+
+		const TCHAR* Suffix = (bNoRegister ? TEXT("_NoRegister") : TEXT(""));
+
+		FString Result;
+		for (UObject* Outer = Item; Outer; Outer = Outer->GetOuter())
+		{
+			if (!Result.IsEmpty())
+			{
+				Result = TEXT("_") + Result;
+			}
+
+			if (Cast<UClass>(Outer) || Cast<UScriptStruct>(Outer))
+			{
+				Result = FNameLookupCPP::GetNameCPP(Cast<UStruct>(Outer)) + Result;
+
+				// Structs can also have UPackage outer.
+				if (Cast<UClass>(Outer) || Cast<UPackage>(Outer->GetOuter()))
+				{
+					break;
+				}
+			}
+			else
+			{
+				Result = Outer->GetName() + Result;
+			}
+		}
+
+		// Can't use long package names in function names.
+		if (Result.StartsWith(TEXT("/Script/"), ESearchCase::CaseSensitive))
+		{
+			Result = FPackageName::GetShortName(Result);
+		}
+
+		const FString ClassString = FNameLookupCPP::GetNameCPP(Item->GetClass());
+		return FString::Printf(TEXT("Z_Construct_%s_%s%s()"), *ClassString, *Result, Suffix);
+	}
+
+	static TMap<FTypeSingletonCacheKey, TUniquePtr<FTypeSingleton>> CacheData;
+	static FRWLock Mutex;
+};
+
+TMap<FTypeSingletonCache::FTypeSingletonCacheKey, TUniquePtr<FTypeSingleton>> FTypeSingletonCache::CacheData;
+FRWLock FTypeSingletonCache::Mutex;
+
+const FString& FNativeClassHeaderGenerator::GetSingletonName(UField* Item, TSet<FString>* UniqueCrossModuleReferences, bool bRequiresValidObject)
+{
+	const FTypeSingleton& Cache = FTypeSingletonCache::Get(Item, bRequiresValidObject);
+
+	// We don't need to export UFunction externs, though we may need the externs for UDelegateFunctions
+	if (UniqueCrossModuleReferences && (!Item->IsA<UFunction>() || Item->IsA<UDelegateFunction>()))
+	{
+		UniqueCrossModuleReferences->Add(Cache.GetExternDecl());
+	}
+
+	return Cache.GetName();
+}
+
+FString FNativeClassHeaderGenerator::GetSingletonNameFuncAddr(UField* Item, TSet<FString>* UniqueCrossModuleReferences, bool bRequiresValidObject)
+{
+	FString Result;
+	if (!Item)
+	{
+		Result = TEXT("nullptr");
+	}
+	else
+	{
+		Result = GetSingletonName(Item, UniqueCrossModuleReferences, bRequiresValidObject).LeftChop(2);
+	}
+	return Result;
+}
+
+void FNativeClassHeaderGenerator::PropertyNew(FOutputDevice& DeclOut, FOutputDevice& Out, FReferenceGatherers& OutReferenceGatherers, FProperty* Prop, const TCHAR* OffsetStr, const TCHAR* Name, const TCHAR* DeclSpaces, const TCHAR* Spaces, const TCHAR* SourceStruct) const
+{
+	FString        PropName             = CreateUTF8LiteralString(FNativeClassHeaderGenerator::GetOverriddenName(Prop));
+	FString        PropNameDep          = Prop->HasAllPropertyFlags(CPF_Deprecated) ? Prop->GetName() + TEXT("_DEPRECATED") : Prop->GetName();
+	const TCHAR*   FPropertyObjectFlags = FClass::IsOwnedByDynamicType(Prop) ? TEXT("RF_Public|RF_Transient") : TEXT("RF_Public|RF_Transient|RF_MarkAsNative");
+	EPropertyFlags PropFlags            = Prop->PropertyFlags & ~CPF_ComputedFlags;
+
+	FString PropTag        = GetGeneratedCodeHashTag(Prop);
+	FString PropNotifyFunc = (Prop->RepNotifyFunc != NAME_None) ? CreateUTF8LiteralString(*Prop->RepNotifyFunc.ToString()) : TEXT("nullptr");
+
+	FString ArrayDim = (Prop->ArrayDim != 1) ? FString::Printf(TEXT("CPP_ARRAY_DIM(%s, %s)"), *PropNameDep, SourceStruct) : TEXT("1");
+
+	FString MetaDataParams = OutputMetaDataCodeForObject(DeclOut, Out, Prop, *FString::Printf(TEXT("%s_MetaData"), Name), DeclSpaces, Spaces);
+
+	FString NameWithoutScope = Name;
+	FString Scope;
+	int32 ScopeIndex = NameWithoutScope.Find(TEXT("::"), ESearchCase::CaseSensitive);
+	if (ScopeIndex != INDEX_NONE)
+	{
+		Scope = NameWithoutScope.Left(ScopeIndex) + TEXT("_");
+		NameWithoutScope.RightChopInline(ScopeIndex + 2, false);
+	}
+
+	if (FByteProperty* TypedProp = CastField<FByteProperty>(Prop))
+	{
+		DeclOut.Logf(TEXT("%sstatic const UE4CodeGen_Private::FBytePropertyParams %s;\r\n"), DeclSpaces, *NameWithoutScope);
+
+		Out.Logf(
+			TEXT("%sconst UE4CodeGen_Private::FBytePropertyParams %s = { %s, %s, (EPropertyFlags)0x%016llx, UE4CodeGen_Private::EPropertyGenFlags::Byte, %s, %s, %s, %s, %s };%s\r\n"),
+			Spaces,
+			Name,
+			*PropName,
+			*PropNotifyFunc,
+			PropFlags,
+			FPropertyObjectFlags,
+			*ArrayDim,
+			OffsetStr,
+			*GetSingletonNameFuncAddr(TypedProp->Enum, OutReferenceGatherers.UniqueCrossModuleReferences),
+			*MetaDataParams,
+			*PropTag
+		);
+
+		return;
+	}
+
+	if (FInt8Property* TypedProp = CastField<FInt8Property>(Prop))
+	{
+		DeclOut.Logf(TEXT("%sstatic const UE4CodeGen_Private::FInt8PropertyParams %s;\r\n"), DeclSpaces, *NameWithoutScope);
+
+		Out.Logf(
+			TEXT("%sconst UE4CodeGen_Private::FInt8PropertyParams %s = { %s, %s, (EPropertyFlags)0x%016llx, UE4CodeGen_Private::EPropertyGenFlags::Int8, %s, %s, %s, %s };%s\r\n"),
+			Spaces,
+			Name,
+			*PropName,
+			*PropNotifyFunc,
+			PropFlags,
+			FPropertyObjectFlags,
+			*ArrayDim,
+			OffsetStr,
+			*MetaDataParams,
+			*PropTag
+		);
+
+		return;
+	}
+
+	if (FInt16Property* TypedProp = CastField<FInt16Property>(Prop))
+	{
+		DeclOut.Logf(TEXT("%sstatic const UE4CodeGen_Private::FInt16PropertyParams %s;\r\n"), DeclSpaces, *NameWithoutScope);
+
+		Out.Logf(
+			TEXT("%sconst UE4CodeGen_Private::FInt16PropertyParams %s = { %s, %s, (EPropertyFlags)0x%016llx, UE4CodeGen_Private::EPropertyGenFlags::Int16, %s, %s, %s, %s };%s\r\n"),
+			Spaces,
+			Name,
+			*PropName,
+			*PropNotifyFunc,
+			PropFlags,
+			FPropertyObjectFlags,
+			*ArrayDim,
+			OffsetStr,
+			*MetaDataParams,
+			*PropTag
+		);
+
+		return;
+	}
+
+	if (FIntProperty* TypedProp = CastField<FIntProperty>(Prop))
+	{
+		const TCHAR* PropTypeName = GUnsizedProperties.Contains(TypedProp) ? TEXT("FUnsizedIntPropertyParams") : TEXT("FIntPropertyParams");
+
+		DeclOut.Logf(TEXT("%sstatic const UE4CodeGen_Private::%s %s;\r\n"), DeclSpaces, PropTypeName, *NameWithoutScope);
+
+		Out.Logf(
+			TEXT("%sconst UE4CodeGen_Private::%s %s = { %s, %s, (EPropertyFlags)0x%016llx, UE4CodeGen_Private::EPropertyGenFlags::Int, %s, %s, %s, %s };%s\r\n"),
+			Spaces,
+			PropTypeName,
+			Name,
+			*PropName,
+			*PropNotifyFunc,
+			PropFlags,
+			FPropertyObjectFlags,
+			*ArrayDim,
+			OffsetStr,
+			*MetaDataParams,
+			*PropTag
+		);
+
+		return;
+	}
+
+	if (FInt64Property* TypedProp = CastField<FInt64Property>(Prop))
+	{
+		DeclOut.Logf(TEXT("%sstatic const UE4CodeGen_Private::FInt64PropertyParams %s;\r\n"), DeclSpaces, *NameWithoutScope);
+
+		Out.Logf(
+			TEXT("%sconst UE4CodeGen_Private::FInt64PropertyParams %s = { %s, %s, (EPropertyFlags)0x%016llx, UE4CodeGen_Private::EPropertyGenFlags::Int64, %s, %s, %s, %s };%s\r\n"),
+			Spaces,
+			Name,
+			*PropName,
+			*PropNotifyFunc,
+			PropFlags,
+			FPropertyObjectFlags,
+			*ArrayDim,
+			OffsetStr,
+			*MetaDataParams,
+			*PropTag
+		);
+
+		return;
+	}
+
+	if (FUInt16Property* TypedProp = CastField<FUInt16Property>(Prop))
+	{
+		DeclOut.Logf(TEXT("%sstatic const UE4CodeGen_Private::FFInt16PropertyParams %s;\r\n"), DeclSpaces, *NameWithoutScope);
+
+		Out.Logf(
+			TEXT("%sconst UE4CodeGen_Private::FFInt16PropertyParams %s = { %s, %s, (EPropertyFlags)0x%016llx, UE4CodeGen_Private::EPropertyGenFlags::UInt16, %s, %s, %s, %s };%s\r\n"),
+			Spaces,
+			Name,
+			*PropName,
+			*PropNotifyFunc,
+			PropFlags,
+			FPropertyObjectFlags,
+			*ArrayDim,
+			OffsetStr,
+			*MetaDataParams,
+			*PropTag
+		);
+
+		return;
+	}
+
+	if (FUInt32Property* TypedProp = CastField<FUInt32Property>(Prop))
+	{
+		const TCHAR* PropTypeName = GUnsizedProperties.Contains(TypedProp) ? TEXT("FUnsizedFIntPropertyParams") : TEXT("FUInt32PropertyParams");
+
+		DeclOut.Logf(TEXT("%sstatic const UE4CodeGen_Private::%s %s;\r\n"), DeclSpaces, PropTypeName, *NameWithoutScope);
+
+		Out.Logf(
+			TEXT("%sconst UE4CodeGen_Private::%s %s = { %s, %s, (EPropertyFlags)0x%016llx, UE4CodeGen_Private::EPropertyGenFlags::UInt32, %s, %s, %s, %s };%s\r\n"),
+			Spaces,
+			PropTypeName,
+			Name,
+			*PropName,
+			*PropNotifyFunc,
+			PropFlags,
+			FPropertyObjectFlags,
+			*ArrayDim,
+			OffsetStr,
+			*MetaDataParams,
+			*PropTag
+		);
+
+		return;
+	}
+
+	if (FUInt64Property* TypedProp = CastField<FUInt64Property>(Prop))
+	{
+		DeclOut.Logf(TEXT("%sstatic const UE4CodeGen_Private::FFInt64PropertyParams %s;\r\n"), DeclSpaces, *NameWithoutScope);
+
+		Out.Logf(
+			TEXT("%sconst UE4CodeGen_Private::FFInt64PropertyParams %s = { %s, %s, (EPropertyFlags)0x%016llx, UE4CodeGen_Private::EPropertyGenFlags::UInt64, %s, %s, %s, %s };%s\r\n"),
+			Spaces,
+			Name,
+			*PropName,
+			*PropNotifyFunc,
+			PropFlags,
+			FPropertyObjectFlags,
+			*ArrayDim,
+			OffsetStr,
+			*MetaDataParams,
+			*PropTag
+		);
+
+		return;
+	}
+
+	if (FFloatProperty* TypedProp = CastField<FFloatProperty>(Prop))
+	{
+		DeclOut.Logf(TEXT("%sstatic const UE4CodeGen_Private::FFloatPropertyParams %s;\r\n"), DeclSpaces, *NameWithoutScope);
+
+		Out.Logf(
+			TEXT("%sconst UE4CodeGen_Private::FFloatPropertyParams %s = { %s, %s, (EPropertyFlags)0x%016llx, UE4CodeGen_Private::EPropertyGenFlags::Float, %s, %s, %s, %s };%s\r\n"),
+			Spaces,
+			Name,
+			*PropName,
+			*PropNotifyFunc,
+			PropFlags,
+			FPropertyObjectFlags,
+			*ArrayDim,
+			OffsetStr,
+			*MetaDataParams,
+			*PropTag
+		);
+
+		return;
+	}
+
+	if (FDoubleProperty* TypedProp = CastField<FDoubleProperty>(Prop))
+	{
+		DeclOut.Logf(TEXT("%sstatic const UE4CodeGen_Private::FDoublePropertyParams %s;\r\n"), DeclSpaces, *NameWithoutScope);
+
+		Out.Logf(
+			TEXT("%sconst UE4CodeGen_Private::FDoublePropertyParams %s = { %s, %s, (EPropertyFlags)0x%016llx, UE4CodeGen_Private::EPropertyGenFlags::Double, %s, %s, %s, %s };%s\r\n"),
+			Spaces,
+			Name,
+			*PropName,
+			*PropNotifyFunc,
+			PropFlags,
+			FPropertyObjectFlags,
+			*ArrayDim,
+			OffsetStr,
+			*MetaDataParams,
+			*PropTag
+		);
+
+		return;
+	}
+
+	if (FBoolProperty* TypedProp = CastField<FBoolProperty>(Prop))
+	{
+		FString OuterSize;
+		FString Setter;
+		if (!Prop->GetOwner<UObject>())
+		{
+			OuterSize = TEXT("0");
+			Setter    = TEXT("nullptr");
+		}
+		else
+		{
+			OuterSize = FString::Printf(TEXT("sizeof(%s)"), SourceStruct);
+
+			DeclOut.Logf(TEXT("%sstatic void %s_SetBit(void* Obj);\r\n"), DeclSpaces, *NameWithoutScope);
+
+			Out.Logf(TEXT("%svoid %s_SetBit(void* Obj)\r\n"), Spaces, Name);
+			Out.Logf(TEXT("%s{\r\n"), Spaces);
+			Out.Logf(TEXT("%s\t((%s*)Obj)->%s%s = 1;\r\n"), Spaces, SourceStruct, *Prop->GetName(), Prop->HasAllPropertyFlags(CPF_Deprecated) ? TEXT("_DEPRECATED") : TEXT(""));
+			Out.Logf(TEXT("%s}\r\n"), Spaces);
+
+			Setter = FString::Printf(TEXT("&%s_SetBit"), Name);
+		}
+
+		DeclOut.Logf(TEXT("%sstatic const UE4CodeGen_Private::FBoolPropertyParams %s;\r\n"), DeclSpaces, *NameWithoutScope);
+
+		Out.Logf(
+			TEXT("%sconst UE4CodeGen_Private::FBoolPropertyParams %s = { %s, %s, (EPropertyFlags)0x%016llx, UE4CodeGen_Private::EPropertyGenFlags::Bool %s, %s, %s, sizeof(%s), %s, %s, %s };%s\r\n"),
+			Spaces,
+			Name,
+			*PropName,
+			*PropNotifyFunc,
+			PropFlags,
+			TypedProp->IsNativeBool() ? TEXT("| UE4CodeGen_Private::EPropertyGenFlags::NativeBool") : TEXT(""),
+			FPropertyObjectFlags,
+			*ArrayDim,
+			*TypedProp->GetCPPType(nullptr, 0),
+			*OuterSize,
+			*Setter,
+			*MetaDataParams,
+			*PropTag
+		);
+
+		return;
+	}
+
+	if (FSoftClassProperty* TypedProp = CastField<FSoftClassProperty>(Prop))
+	{
+		DeclOut.Logf(TEXT("%sstatic const UE4CodeGen_Private::FSoftClassPropertyParams %s;\r\n"), DeclSpaces, *NameWithoutScope);
+
+		Out.Logf(
+			TEXT("%sconst UE4CodeGen_Private::FSoftClassPropertyParams %s = { %s, %s, (EPropertyFlags)0x%016llx, UE4CodeGen_Private::EPropertyGenFlags::SoftClass, %s, %s, %s, %s, %s };%s\r\n"),
+			Spaces,
+			Name,
+			*PropName,
+			*PropNotifyFunc,
+			PropFlags,
+			FPropertyObjectFlags,
+			*ArrayDim,
+			OffsetStr,
+			*GetSingletonNameFuncAddr(TypedProp->MetaClass, OutReferenceGatherers.UniqueCrossModuleReferences, false),
+			*MetaDataParams,
+			*PropTag
+		);
+
+		return;
+	}
+
+	if (FWeakObjectProperty* TypedProp = CastField<FWeakObjectProperty>(Prop))
+	{
+		DeclOut.Logf(TEXT("%sstatic const UE4CodeGen_Private::FWeakObjectPropertyParams %s;\r\n"), DeclSpaces, *NameWithoutScope);
+
+		Out.Logf(
+			TEXT("%sconst UE4CodeGen_Private::FWeakObjectPropertyParams %s = { %s, %s, (EPropertyFlags)0x%016llx, UE4CodeGen_Private::EPropertyGenFlags::WeakObject, %s, %s, %s, %s, %s };%s\r\n"),
+			Spaces,
+			Name,
+			*PropName,
+			*PropNotifyFunc,
+			PropFlags,
+			FPropertyObjectFlags,
+			*ArrayDim,
+			OffsetStr,
+			*GetSingletonNameFuncAddr(TypedProp->PropertyClass, OutReferenceGatherers.UniqueCrossModuleReferences, false),
+			*MetaDataParams,
+			*PropTag
+		);
+
+		return;
+	}
+
+	if (FLazyObjectProperty* TypedProp = CastField<FLazyObjectProperty>(Prop))
+	{
+		DeclOut.Logf(TEXT("%sstatic const UE4CodeGen_Private::FLazyObjectPropertyParams %s;\r\n"), DeclSpaces, *NameWithoutScope);
+
+		Out.Logf(
+			TEXT("%sconst UE4CodeGen_Private::FLazyObjectPropertyParams %s = { %s, %s, (EPropertyFlags)0x%016llx, UE4CodeGen_Private::EPropertyGenFlags::LazyObject, %s, %s, %s, %s, %s };%s\r\n"),
+			Spaces,
+			Name,
+			*PropName,
+			*PropNotifyFunc,
+			PropFlags,
+			FPropertyObjectFlags,
+			*ArrayDim,
+			OffsetStr,
+			*GetSingletonNameFuncAddr(TypedProp->PropertyClass, OutReferenceGatherers.UniqueCrossModuleReferences, false),
+			*MetaDataParams,
+			*PropTag
+		);
+
+		return;
+	}
+
+	if (FSoftObjectProperty* TypedProp = CastField<FSoftObjectProperty>(Prop))
+	{
+		DeclOut.Logf(TEXT("%sstatic const UE4CodeGen_Private::FSoftObjectPropertyParams %s;\r\n"), DeclSpaces, *NameWithoutScope);
+
+		Out.Logf(
+			TEXT("%sconst UE4CodeGen_Private::FSoftObjectPropertyParams %s = { %s, %s, (EPropertyFlags)0x%016llx, UE4CodeGen_Private::EPropertyGenFlags::SoftObject, %s, %s, %s, %s, %s };%s\r\n"),
+			Spaces,
+			Name,
+			*PropName,
+			*PropNotifyFunc,
+			PropFlags,
+			FPropertyObjectFlags,
+			*ArrayDim,
+			OffsetStr,
+			*GetSingletonNameFuncAddr(TypedProp->PropertyClass, OutReferenceGatherers.UniqueCrossModuleReferences, false),
+			*MetaDataParams,
+			*PropTag
+		);
+
+		return;
+	}
+
+	if (FClassProperty* TypedProp = CastField<FClassProperty>(Prop))
+	{
+		DeclOut.Logf(TEXT("%sstatic const UE4CodeGen_Private::FClassPropertyParams %s;\r\n"), DeclSpaces, *NameWithoutScope);
+
+		Out.Logf(
+			TEXT("%sconst UE4CodeGen_Private::FClassPropertyParams %s = { %s, %s, (EPropertyFlags)0x%016llx, UE4CodeGen_Private::EPropertyGenFlags::Class, %s, %s, %s, %s, %s, %s };%s\r\n"),
+			Spaces,
+			Name,
+			*PropName,
+			*PropNotifyFunc,
+			PropFlags,
+			FPropertyObjectFlags,
+			*ArrayDim,
+			OffsetStr,
+			*GetSingletonNameFuncAddr(TypedProp->MetaClass, OutReferenceGatherers.UniqueCrossModuleReferences, false),
+			*GetSingletonNameFuncAddr(TypedProp->PropertyClass, OutReferenceGatherers.UniqueCrossModuleReferences, false),
+			*MetaDataParams,
+			*PropTag
+		);
+
+		return;
+	}
+
+	if (FObjectProperty* TypedProp = CastField<FObjectProperty>(Prop))
+	{
+		DeclOut.Logf(TEXT("%sstatic const UE4CodeGen_Private::FObjectPropertyParams %s;\r\n"), DeclSpaces, *NameWithoutScope);
+
+		Out.Logf(
+			TEXT("%sconst UE4CodeGen_Private::FObjectPropertyParams %s = { %s, %s, (EPropertyFlags)0x%016llx, UE4CodeGen_Private::EPropertyGenFlags::Object, %s, %s, %s, %s, %s };%s\r\n"),
+			Spaces,
+			Name,
+			*PropName,
+			*PropNotifyFunc,
+			PropFlags,
+			FPropertyObjectFlags,
+			*ArrayDim,
+			OffsetStr,
+			*GetSingletonNameFuncAddr(TypedProp->PropertyClass, OutReferenceGatherers.UniqueCrossModuleReferences, false),
+			*MetaDataParams,
+			*PropTag
+		);
+
+		return;
+	}
+
+	if (FInterfaceProperty* TypedProp = CastField<FInterfaceProperty>(Prop))
+	{
+		DeclOut.Logf(TEXT("%sstatic const UE4CodeGen_Private::FInterfacePropertyParams %s;\r\n"), DeclSpaces, *NameWithoutScope);
+
+		Out.Logf(
+			TEXT("%sconst UE4CodeGen_Private::FInterfacePropertyParams %s = { %s, %s, (EPropertyFlags)0x%016llx, UE4CodeGen_Private::EPropertyGenFlags::Interface, %s, %s, %s, %s, %s };%s\r\n"),
+			Spaces,
+			Name,
+			*PropName,
+			*PropNotifyFunc,
+			PropFlags,
+			FPropertyObjectFlags,
+			*ArrayDim,
+			OffsetStr,
+			*GetSingletonNameFuncAddr(TypedProp->InterfaceClass, OutReferenceGatherers.UniqueCrossModuleReferences, false),
+			*MetaDataParams,
+			*PropTag
+		);
+
+		return;
+	}
+
+	if (FNameProperty* TypedProp = CastField<FNameProperty>(Prop))
+	{
+		DeclOut.Logf(TEXT("%sstatic const UE4CodeGen_Private::FNamePropertyParams %s;\r\n"), DeclSpaces, *NameWithoutScope);
+
+		Out.Logf(
+			TEXT("%sconst UE4CodeGen_Private::FNamePropertyParams %s = { %s, %s, (EPropertyFlags)0x%016llx, UE4CodeGen_Private::EPropertyGenFlags::Name, %s, %s, %s, %s };%s\r\n"),
+			Spaces,
+			Name,
+			*PropName,
+			*PropNotifyFunc,
+			PropFlags,
+			FPropertyObjectFlags,
+			*ArrayDim,
+			OffsetStr,
+			*MetaDataParams,
+			*PropTag
+		);
+
+		return;
+	}
+
+	if (FStrProperty* TypedProp = CastField<FStrProperty>(Prop))
+	{
+		DeclOut.Logf(TEXT("%sstatic const UE4CodeGen_Private::FStrPropertyParams %s;\r\n"), DeclSpaces, *NameWithoutScope);
+
+		Out.Logf(
+			TEXT("%sconst UE4CodeGen_Private::FStrPropertyParams %s = { %s, %s, (EPropertyFlags)0x%016llx, UE4CodeGen_Private::EPropertyGenFlags::Str, %s, %s, %s, %s };%s\r\n"),
+			Spaces,
+			Name,
+			*PropName,
+			*PropNotifyFunc,
+			PropFlags,
+			FPropertyObjectFlags,
+			*ArrayDim,
+			OffsetStr,
+			*MetaDataParams,
+			*PropTag
+		);
+
+		return;
+	}
+
+	if (FArrayProperty* TypedProp = CastField<FArrayProperty>(Prop))
+	{
+		DeclOut.Logf(TEXT("%sstatic const UE4CodeGen_Private::FArrayPropertyParams %s;\r\n"), DeclSpaces, *NameWithoutScope);
+
+		Out.Logf(
+			TEXT("%sconst UE4CodeGen_Private::FArrayPropertyParams %s = { %s, %s, (EPropertyFlags)0x%016llx, UE4CodeGen_Private::EPropertyGenFlags::Array, %s, %s, %s, %s, %s };%s\r\n"),
+			Spaces,
+			Name,
+			*PropName,
+			*PropNotifyFunc,
+			PropFlags,
+			FPropertyObjectFlags,
+			*ArrayDim,
+			OffsetStr,
+			GPropertyUsesMemoryImageAllocator.Contains(TypedProp) ? TEXT("EArrayPropertyFlags::UsesMemoryImageAllocator") : TEXT("EArrayPropertyFlags::None"),
+			*MetaDataParams,
+			*PropTag
+		);
+
+		return;
+	}
+
+	if (FMapProperty* TypedProp = CastField<FMapProperty>(Prop))
+	{
+		DeclOut.Logf(TEXT("%sstatic const UE4CodeGen_Private::FMapPropertyParams %s;\r\n"), DeclSpaces, *NameWithoutScope);
+
+		Out.Logf(
+			TEXT("%sconst UE4CodeGen_Private::FMapPropertyParams %s = { %s, %s, (EPropertyFlags)0x%016llx, UE4CodeGen_Private::EPropertyGenFlags::Map, %s, %s, %s, %s, %s };%s\r\n"),
+			Spaces,
+			Name,
+			*PropName,
+			*PropNotifyFunc,
+			PropFlags,
+			FPropertyObjectFlags,
+			*ArrayDim,
+			OffsetStr,
+			GPropertyUsesMemoryImageAllocator.Contains(TypedProp) ? TEXT("EMapPropertyFlags::UsesMemoryImageAllocator") : TEXT("EMapPropertyFlags::None"),
+			*MetaDataParams,
+			*PropTag
+		);
+
+		return;
+	}
+
+	if (FSetProperty* TypedProp = CastField<FSetProperty>(Prop))
+	{
+		DeclOut.Logf(TEXT("%sstatic const UE4CodeGen_Private::FSetPropertyParams %s;\r\n"), DeclSpaces, *NameWithoutScope);
+
+		Out.Logf(
+			TEXT("%sconst UE4CodeGen_Private::FSetPropertyParams %s = { %s, %s, (EPropertyFlags)0x%016llx, UE4CodeGen_Private::EPropertyGenFlags::Set, %s, %s, %s, %s };%s\r\n"),
+			Spaces,
+			Name,
+			*PropName,
+			*PropNotifyFunc,
+			PropFlags,
+			FPropertyObjectFlags,
+			*ArrayDim,
+			OffsetStr,
+			*MetaDataParams,
+			*PropTag
+		);
+
+		return;
+	}
+
+	if (FStructProperty* TypedProp = CastField<FStructProperty>(Prop))
+	{
+		DeclOut.Logf(TEXT("%sstatic const UE4CodeGen_Private::FStructPropertyParams %s;\r\n"), DeclSpaces, *NameWithoutScope);
+
+		Out.Logf(
+			TEXT("%sconst UE4CodeGen_Private::FStructPropertyParams %s = { %s, %s, (EPropertyFlags)0x%016llx, UE4CodeGen_Private::EPropertyGenFlags::Struct, %s, %s, %s, %s, %s };%s\r\n"),
+			Spaces,
+			Name,
+			*PropName,
+			*PropNotifyFunc,
+			PropFlags,
+			FPropertyObjectFlags,
+			*ArrayDim,
+			OffsetStr,
+			*GetSingletonNameFuncAddr(TypedProp->Struct, OutReferenceGatherers.UniqueCrossModuleReferences),
+			*MetaDataParams,
+			*PropTag
+		);
+
+		return;
+	}
+
+	if (FDelegateProperty* TypedProp = CastField<FDelegateProperty>(Prop))
+	{
+		DeclOut.Logf(TEXT("%sstatic const UE4CodeGen_Private::FDelegatePropertyParams %s;\r\n"), DeclSpaces, *NameWithoutScope);
+
+		Out.Logf(
+			TEXT("%sconst UE4CodeGen_Private::FDelegatePropertyParams %s = { %s, %s, (EPropertyFlags)0x%016llx, UE4CodeGen_Private::EPropertyGenFlags::Delegate, %s, %s, %s, %s, %s };%s\r\n"),
+			Spaces,
+			Name,
+			*PropName,
+			*PropNotifyFunc,
+			PropFlags,
+			FPropertyObjectFlags,
+			*ArrayDim,
+			OffsetStr,
+			*GetSingletonNameFuncAddr(TypedProp->SignatureFunction, OutReferenceGatherers.UniqueCrossModuleReferences),
+			*MetaDataParams,
+			*PropTag
+		);
+
+		return;
+	}
+
+	if (FMulticastDelegateProperty* TypedProp = CastField<FMulticastDelegateProperty>(Prop))
+	{
+		DeclOut.Logf(TEXT("%sstatic const UE4CodeGen_Private::FMulticastDelegatePropertyParams %s;\r\n"), DeclSpaces, *NameWithoutScope);
+
+		Out.Logf(
+			TEXT("%sconst UE4CodeGen_Private::FMulticastDelegatePropertyParams %s = { %s, %s, (EPropertyFlags)0x%016llx, UE4CodeGen_Private::EPropertyGenFlags::%sMulticastDelegate, %s, %s, %s, %s, %s };%s\r\n"),
+			Spaces,
+			Name,
+			*PropName,
+			*PropNotifyFunc,
+			PropFlags,
+			(TypedProp->IsA<FMulticastInlineDelegateProperty>() ? TEXT("Inline") : TEXT("Sparse")),
+			FPropertyObjectFlags,
+			*ArrayDim,
+			OffsetStr,
+			*GetSingletonNameFuncAddr(TypedProp->SignatureFunction, OutReferenceGatherers.UniqueCrossModuleReferences),
+			*MetaDataParams,
+			*PropTag
+		);
+
+		return;
+	}
+
+	if (FTextProperty* TypedProp = CastField<FTextProperty>(Prop))
+	{
+		DeclOut.Logf(TEXT("%sstatic const UE4CodeGen_Private::FTextPropertyParams %s;\r\n"), DeclSpaces, *NameWithoutScope);
+
+		Out.Logf(
+			TEXT("%sconst UE4CodeGen_Private::FTextPropertyParams %s = { %s, %s, (EPropertyFlags)0x%016llx, UE4CodeGen_Private::EPropertyGenFlags::Text, %s, %s, %s, %s };%s\r\n"),
+			Spaces,
+			Name,
+			*PropName,
+			*PropNotifyFunc,
+			PropFlags,
+			FPropertyObjectFlags,
+			*ArrayDim,
+			OffsetStr,
+			*MetaDataParams,
+			*PropTag
+		);
+
+		return;
+	}
+
+	if (FEnumProperty* TypedProp = CastField<FEnumProperty>(Prop))
+	{
+		DeclOut.Logf(TEXT("%sstatic const UE4CodeGen_Private::FEnumPropertyParams %s;\r\n"), DeclSpaces, *NameWithoutScope);
+
+		Out.Logf(
+			TEXT("%sconst UE4CodeGen_Private::FEnumPropertyParams %s = { %s, %s, (EPropertyFlags)0x%016llx, UE4CodeGen_Private::EPropertyGenFlags::Enum, %s, %s, %s, %s, %s };%s\r\n"),
+			Spaces,
+			Name,
+			*PropName,
+			*PropNotifyFunc,
+			PropFlags,
+			FPropertyObjectFlags,
+			*ArrayDim,
+			OffsetStr,
+			*GetSingletonNameFuncAddr(TypedProp->Enum, OutReferenceGatherers.UniqueCrossModuleReferences),
+			*MetaDataParams,
+			*PropTag
+		);
+
+		return;
+	}
+
+	if (FFieldPathProperty* TypedProp = CastField<FFieldPathProperty>(Prop))
+	{
+		DeclOut.Logf(TEXT("%sstatic const UE4CodeGen_Private::FFieldPathPropertyParams %s;\r\n"), DeclSpaces, *NameWithoutScope);
+
+		Out.Logf(
+			TEXT("%sconst UE4CodeGen_Private::FFieldPathPropertyParams %s = { %s, %s, (EPropertyFlags)0x%016llx, UE4CodeGen_Private::EPropertyGenFlags::FieldPath, %s, %s, %s, %s, %s };%s\r\n"),
+			Spaces,
+			Name,
+			*PropName,
+			*PropNotifyFunc,
+			PropFlags,
+			FPropertyObjectFlags,
+			*ArrayDim,
+			OffsetStr,
+			*FString::Printf(TEXT("&F%s::StaticClass"), *TypedProp->PropertyClass->GetName()),
+			*MetaDataParams,
+			*PropTag
+		);
+
+		return;
+	}
+
+	// Unhandled type
+	check(false);
+}
+
+bool IsEditorOnlyDataProperty(FProperty* Prop)
+{
+	while (Prop)
+	{
+		if (Prop->IsEditorOnlyProperty())
+		{
+			return true;
+		}
+
+		Prop = Prop->GetOwner<FProperty>();
+	}
+
+	return false;
+}
+
+TTuple<FString, FString> FNativeClassHeaderGenerator::OutputProperties(FOutputDevice& DeclOut, FOutputDevice& Out, FReferenceGatherers& OutReferenceGatherers, const TCHAR* Scope, const TArray<FProperty*>& Properties, const TCHAR* DeclSpaces, const TCHAR* Spaces) const
+{
+	if (Properties.Num() == 0)
+	{
+		return TTuple<FString, FString>(TEXT("nullptr"), TEXT("0"));
+	}
+
+	TArray<FPropertyNamePointerPair> PropertyNamesAndPointers;
+	bool bHasAllEditorOnlyDataProperties = true;
+
+	{
+		FMacroBlockEmitter WithEditorOnlyMacroEmitter(Out, TEXT("WITH_EDITORONLY_DATA"));
+		FMacroBlockEmitter WithEditorOnlyMacroEmitterDecl(DeclOut, TEXT("WITH_EDITORONLY_DATA"));
+
+		for (FProperty* Prop : Properties)
+		{
+			bool bRequiresHasEditorOnlyMacro = IsEditorOnlyDataProperty(Prop);
+			if (!bRequiresHasEditorOnlyMacro)
+			{
+				bHasAllEditorOnlyDataProperties = false;
+			}
+
+			WithEditorOnlyMacroEmitter(bRequiresHasEditorOnlyMacro);
+			WithEditorOnlyMacroEmitterDecl(bRequiresHasEditorOnlyMacro);
+			OutputProperty(DeclOut, Out, OutReferenceGatherers, Scope, PropertyNamesAndPointers, Prop, DeclSpaces, Spaces);
+		}
+
+		WithEditorOnlyMacroEmitter(bHasAllEditorOnlyDataProperties);
+		WithEditorOnlyMacroEmitterDecl(bHasAllEditorOnlyDataProperties);
+		DeclOut.Logf(TEXT("%sstatic const UE4CodeGen_Private::FPropertyParamsBase* const PropPointers[];\r\n"), DeclSpaces);
+		Out.Logf(TEXT("%sconst UE4CodeGen_Private::FPropertyParamsBase* const %sPropPointers[] = {\r\n"), Spaces, Scope);
+
+		for (const FPropertyNamePointerPair& PropNameAndPtr : PropertyNamesAndPointers)
+		{
+			bool bRequiresHasEditorOnlyMacro = IsEditorOnlyDataProperty(PropNameAndPtr.Prop);
+
+			WithEditorOnlyMacroEmitter(bRequiresHasEditorOnlyMacro);
+			WithEditorOnlyMacroEmitterDecl(bRequiresHasEditorOnlyMacro);
+			Out.Logf(TEXT("%s\t(const UE4CodeGen_Private::FPropertyParamsBase*)&%s,\r\n"), Spaces, *PropNameAndPtr.Name);
+		}
+
+		WithEditorOnlyMacroEmitter(bHasAllEditorOnlyDataProperties);
+		WithEditorOnlyMacroEmitterDecl(bHasAllEditorOnlyDataProperties);
+		Out.Logf(TEXT("%s};\r\n"), Spaces);
+	}
+
+	if (bHasAllEditorOnlyDataProperties)
+	{
+		return TTuple<FString, FString>(
+			FString::Printf(TEXT("IF_WITH_EDITORONLY_DATA(%sPropPointers, nullptr)"), Scope),
+			FString::Printf(TEXT("IF_WITH_EDITORONLY_DATA(UE_ARRAY_COUNT(%sPropPointers), 0)"), Scope)
+		);
+	}
+	else
+	{
+		return TTuple<FString, FString>(
+			FString::Printf(TEXT("%sPropPointers"), Scope),
+			FString::Printf(TEXT("UE_ARRAY_COUNT(%sPropPointers)"), Scope)
+		);
+	}
+}
+
+inline FString GetEventStructParamsName(UObject* Outer, const TCHAR* FunctionName)
+{
+	FString OuterName;
+	if (Outer->IsA<UClass>())
+	{
+		OuterName = ((UClass*)Outer)->GetName();
+	}
+	else if (Outer->IsA<UPackage>())
+	{
+		OuterName = ((UPackage*)Outer)->GetName();
+		OuterName.ReplaceInline(TEXT("/"), TEXT("_"), ESearchCase::CaseSensitive);
+	}
+	else
+	{
+		FError::Throwf(TEXT("Unrecognized outer type"));
+	}
+
+	FString Result = FString::Printf(TEXT("%s_event%s_Parms"), *OuterName, FunctionName);
+	if (Result.Len() && FChar::IsDigit(Result[0]))
+	{
+		Result.InsertAt(0, TCHAR('_'));
+	}
+	return Result;
+}
+
+void FNativeClassHeaderGenerator::OutputProperty(FOutputDevice& DeclOut, FOutputDevice& Out, FReferenceGatherers& OutReferenceGatherers, const TCHAR* Scope, TArray<FPropertyNamePointerPair>& PropertyNamesAndPointers, FProperty* Prop, const TCHAR* DeclSpaces, const TCHAR* Spaces) const
+{
+	// Helper to handle the creation of the underlying properties if they're enum properties
+	auto HandleUnderlyingEnumProperty = [this, &PropertyNamesAndPointers, &DeclOut, &Out, &OutReferenceGatherers, DeclSpaces, Spaces](FProperty* LocalProp, FString&& InOuterName)
+	{
+		if (FEnumProperty* EnumProp = CastField<FEnumProperty>(LocalProp))
+		{
+			FString PropVarName = InOuterName + TEXT("_Underlying");
+
+			PropertyNew(DeclOut, Out, OutReferenceGatherers, EnumProp->UnderlyingProp, TEXT("0"), *PropVarName, DeclSpaces, Spaces);
+			PropertyNamesAndPointers.Emplace(MoveTemp(PropVarName), EnumProp->UnderlyingProp);
+		}
+
+		PropertyNamesAndPointers.Emplace(MoveTemp(InOuterName), LocalProp);
+	};
+
+	if (FArrayProperty* ArrayProperty = CastField<FArrayProperty>(Prop))
+	{
+		FString InnerVariableName = FString::Printf(TEXT("%sNewProp_%s_Inner"), Scope, *ArrayProperty->Inner->GetName());
+
+		HandleUnderlyingEnumProperty(ArrayProperty->Inner, CopyTemp(InnerVariableName));
+		PropertyNew(DeclOut, Out, OutReferenceGatherers, ArrayProperty->Inner, TEXT("0"), *InnerVariableName, DeclSpaces, Spaces);
+	}
+
+	else if (FMapProperty* MapProperty = CastField<FMapProperty>(Prop))
+	{
+		FProperty* Key   = MapProperty->KeyProp;
+		FProperty* Value = MapProperty->ValueProp;
+
+		FString KeyVariableName   = FString::Printf(TEXT("%sNewProp_%s_KeyProp"), Scope, *Key->GetName());
+		FString ValueVariableName = FString::Printf(TEXT("%sNewProp_%s_ValueProp"), Scope, *Value->GetName());
+
+		HandleUnderlyingEnumProperty(Value, CopyTemp(ValueVariableName));
+		PropertyNew(DeclOut, Out, OutReferenceGatherers, Value, TEXT("1"), *ValueVariableName, DeclSpaces, Spaces);
+
+		HandleUnderlyingEnumProperty(Key, CopyTemp(KeyVariableName));
+		PropertyNew(DeclOut, Out, OutReferenceGatherers, Key, TEXT("0"), *KeyVariableName, DeclSpaces, Spaces);
+	}
+
+	else if (FSetProperty* SetProperty = CastField<FSetProperty>(Prop))
+	{
+		FProperty* Inner = SetProperty->ElementProp;
+
+		FString ElementVariableName = FString::Printf(TEXT("%sNewProp_%s_ElementProp"), Scope, *Inner->GetName());
+
+		HandleUnderlyingEnumProperty(Inner, CopyTemp(ElementVariableName));
+		PropertyNew(DeclOut, Out, OutReferenceGatherers, Inner, TEXT("0"), *ElementVariableName, DeclSpaces, Spaces);
+	}
+
+	{
+		FString SourceStruct;
+		if (UFunction* Function = Prop->GetOwner<UFunction>())
+		{
+			while (Function->GetSuperFunction())
+			{
+				Function = Function->GetSuperFunction();
+			}
+			FString FunctionName = Function->GetName();
+			if (Function->HasAnyFunctionFlags(FUNC_Delegate))
+			{
+				FunctionName.LeftChopInline(HEADER_GENERATED_DELEGATE_SIGNATURE_SUFFIX_LENGTH, false);
+			}
+
+			SourceStruct = GetEventStructParamsName(Function->GetOuter(), *FunctionName);
+		}
+		else
+		{
+			SourceStruct = FNameLookupCPP::GetNameCPP(CastChecked<UStruct>(Prop->GetOwner<UObject>()));
+		}
+
+		FString PropName = Prop->GetName();
+		FString PropVariableName = FString::Printf(TEXT("%sNewProp_%s"), Scope, *PropName);
+
+		if (Prop->HasAllPropertyFlags(CPF_Deprecated))
+		{
+			PropName += TEXT("_DEPRECATED");
+		}
+
+		FString PropMacroOuterClass = FString::Printf(TEXT("STRUCT_OFFSET(%s, %s)"), *SourceStruct, *PropName);
+
+		HandleUnderlyingEnumProperty(Prop, CopyTemp(PropVariableName));
+		PropertyNew(DeclOut, Out, OutReferenceGatherers, Prop, *PropMacroOuterClass, *PropVariableName, DeclSpaces, Spaces, *SourceStruct);
+	}
+}
+
+static bool IsAlwaysAccessible(UScriptStruct* Script)
+{
+	FName ToTest = Script->GetFName();
+	if (ToTest == NAME_Matrix)
+	{
+		return false; // special case, the C++ FMatrix does not have the same members.
+	}
+	bool Result = Script->HasDefaults(); // if we have cpp struct ops in it for UHT, then we can assume it is always accessible
+	if( ToTest == NAME_Plane
+		||	ToTest == NAME_Vector
+		||	ToTest == NAME_Vector4
+		||	ToTest == NAME_Quat
+		||	ToTest == NAME_Color
+		)
+	{
+		check(Result);
+	}
+	return Result;
+}
+
+static void FindNoExportStructsRecursive(TArray<UScriptStruct*>& Structs, UStruct* Start)
+{
+	while (Start)
+	{
+		if (UScriptStruct* StartScript = Cast<UScriptStruct>(Start))
+		{
+			if (StartScript->StructFlags & STRUCT_Native)
+			{
+				break;
+			}
+
+			if (!IsAlwaysAccessible(StartScript)) // these are a special cases that already exists and if wrong if exported naively
+			{
+				// this will topologically sort them in reverse order
+				Structs.Remove(StartScript);
+				Structs.Add(StartScript);
+			}
+		}
+
+		for (FProperty* Prop : TFieldRange<FProperty>(Start, EFieldIteratorFlags::ExcludeSuper))
+		{
+			if (FStructProperty* StructProp = CastField<FStructProperty>(Prop))
+			{
+				FindNoExportStructsRecursive(Structs, StructProp->Struct);
+			}
+			else if (FArrayProperty* ArrayProp = CastField<FArrayProperty>(Prop))
+			{
+				if (FStructProperty* InnerStructProp = CastField<FStructProperty>(ArrayProp->Inner))
+				{
+					FindNoExportStructsRecursive(Structs, InnerStructProp->Struct);
+				}
+			}
+			else if (FMapProperty* MapProp = CastField<FMapProperty>(Prop))
+			{
+				if (FStructProperty* KeyStructProp = CastField<FStructProperty>(MapProp->KeyProp))
+				{
+					FindNoExportStructsRecursive(Structs, KeyStructProp->Struct);
+				}
+				if (FStructProperty* ValueStructProp = CastField<FStructProperty>(MapProp->ValueProp))
+				{
+					FindNoExportStructsRecursive(Structs, ValueStructProp->Struct);
+				}
+			}
+			else if (FSetProperty* SetProp = CastField<FSetProperty>(Prop))
+			{
+				if (FStructProperty* ElementStructProp = CastField<FStructProperty>(SetProp->ElementProp))
+				{
+					FindNoExportStructsRecursive(Structs, ElementStructProp->Struct);
+				}
+			}
+		}
+		Start = Start->GetSuperStruct();
+	}
+}
+
+static TArray<UScriptStruct*> FindNoExportStructs(UStruct* Start)
+{
+	TArray<UScriptStruct*> Result;
+	FindNoExportStructsRecursive(Result, Start);
+
+	// These come out in reverse order of topology so reverse them
+	Algo::Reverse(Result);
+
+	return Result;
+}
+
+struct FPackageSingletonStrings
+{
+	FPackageSingletonStrings(FString&& InPackageSingletonName)
+		: PackageSingletonName(MoveTemp(InPackageSingletonName))
+		, PackageUniqueCrossModuleReference(FString::Printf(TEXT("\tUPackage* %s;\r\n"), *PackageSingletonName))
+	{
+	}
+
+	FString PackageSingletonName;
+	FString PackageUniqueCrossModuleReference;
+};
+
+static TMap<const UPackage*, TUniquePtr<FPackageSingletonStrings>> PackageSingletonNames;
+static FRWLock PackageSingletonNamesLock;
+
+const FString& FNativeClassHeaderGenerator::GetPackageSingletonName(const UPackage* InPackage, TSet<FString>* UniqueCrossModuleReferences)
+{
+	FRWScopeLock Lock(PackageSingletonNamesLock, SLT_ReadOnly);
+
+	TUniquePtr<FPackageSingletonStrings>* PackageSingletonStrings = PackageSingletonNames.Find(InPackage);
+	if (PackageSingletonStrings == nullptr)
+	{
+		FString PackageName = InPackage->GetName();
+		PackageName.ReplaceInline(TEXT("/"), TEXT("_"), ESearchCase::CaseSensitive);
+
+		TUniquePtr<FPackageSingletonStrings> NewPackageSingletonStrings(MakeUnique<FPackageSingletonStrings>(FString::Printf(TEXT("Z_Construct_UPackage_%s()"), *PackageName)));
+
+		Lock.ReleaseReadOnlyLockAndAcquireWriteLock_USE_WITH_CAUTION();
+		
+		// Check the map again in case another thread had also been waiting on writing this data and got the write lock first
+		PackageSingletonStrings = PackageSingletonNames.Find(InPackage);
+
+		if (PackageSingletonStrings == nullptr)
+		{
+			PackageSingletonStrings = &PackageSingletonNames.Add(InPackage, MoveTemp(NewPackageSingletonStrings));
+		}
+	}
+
+	if (UniqueCrossModuleReferences)
+	{
+		UniqueCrossModuleReferences->Add((*PackageSingletonStrings)->PackageUniqueCrossModuleReference);
+	}
+
+	return (*PackageSingletonStrings)->PackageSingletonName;
+}
+
+void FNativeClassHeaderGenerator::ExportGeneratedPackageInitCode(FOutputDevice& Out, const TCHAR* InDeclarations, const UPackage* InPackage, uint32 Hash)
+{
+	const FString& SingletonName = GetPackageSingletonName(InPackage, nullptr);
+
+	TArray<UField*>* SingletonsToOutput = GPackageSingletons.Find(InPackage);
+	if (SingletonsToOutput)
+	{
+		Algo::Sort(*SingletonsToOutput, [](UField* A, UField* B)
+		{
+			// Structs before delegates then UniqueId order
+			return (uint64(A->IsA<UDelegateFunction>()) << 32) + A->GetUniqueID() <
+			       (uint64(B->IsA<UDelegateFunction>()) << 32) + B->GetUniqueID();
+		});
+
+		for (UField* ScriptType : *SingletonsToOutput)
+		{
+			Out.Log(FTypeSingletonCache::Get(ScriptType, true).GetExternDecl());
+		}
+	}
+
+	FOutputDeviceNull OutputDeviceNull;
+	FString MetaDataParams = OutputMetaDataCodeForObject(OutputDeviceNull, Out, const_cast<UPackage*>(InPackage), TEXT("Package_MetaDataParams"), TEXT(""), TEXT("\t\t\t"));
+
+	Out.Logf(TEXT("\tUPackage* %s\r\n"), *SingletonName);
+	Out.Logf(TEXT("\t{\r\n"));
+	Out.Logf(TEXT("\t\tstatic UPackage* ReturnPackage = nullptr;\r\n"));
+	Out.Logf(TEXT("\t\tif (!ReturnPackage)\r\n"));
+	Out.Logf(TEXT("\t\t{\r\n"));
+
+	const TCHAR* SingletonArray;
+	const TCHAR* SingletonCount;
+	if (SingletonsToOutput)
+	{
+		Out.Logf(TEXT("\t\t\tstatic UObject* (*const SingletonFuncArray[])() = {\r\n"));
+		for (UField* ScriptType : *SingletonsToOutput)
+		{
+			const FString Name = FTypeSingletonCache::Get(ScriptType, true).GetName().LeftChop(2);
+
+			Out.Logf(TEXT("\t\t\t\t(UObject* (*)())%s,\r\n"), *Name);
+		}
+		Out.Logf(TEXT("\t\t\t};\r\n"));
+
+		SingletonArray = TEXT("SingletonFuncArray");
+		SingletonCount = TEXT("UE_ARRAY_COUNT(SingletonFuncArray)");
+	}
+	else
+	{
+		SingletonArray = TEXT("nullptr");
+		SingletonCount = TEXT("0");
+	}
+
+	Out.Logf(TEXT("\t\t\tstatic const UE4CodeGen_Private::FPackageParams PackageParams = {\r\n"));
+	Out.Logf(TEXT("\t\t\t\t%s,\r\n"), *CreateUTF8LiteralString(InPackage->GetName()));
+	Out.Logf(TEXT("\t\t\t\t%s,\r\n"), SingletonArray);
+	Out.Logf(TEXT("\t\t\t\t%s,\r\n"), SingletonCount);
+	Out.Logf(TEXT("\t\t\t\tPKG_CompiledIn | 0x%08X,\r\n"), InPackage->GetPackageFlags() & (PKG_ClientOptional | PKG_ServerSideOnly | PKG_EditorOnly | PKG_Developer | PKG_UncookedOnly));
+	Out.Logf(TEXT("\t\t\t\t0x%08X,\r\n"), Hash);
+	Out.Logf(TEXT("\t\t\t\t0x%08X,\r\n"), GenerateTextHash(InDeclarations));
+	Out.Logf(TEXT("\t\t\t\t%s\r\n"), *MetaDataParams);
+	Out.Logf(TEXT("\t\t\t};\r\n"));
+	Out.Logf(TEXT("\t\t\tUE4CodeGen_Private::ConstructUPackage(ReturnPackage, PackageParams);\r\n"));
+	Out.Logf(TEXT("\t\t}\r\n"));
+	Out.Logf(TEXT("\t\treturn ReturnPackage;\r\n"));
+	Out.Logf(TEXT("\t}\r\n"));
+}
+
+void FNativeClassHeaderGenerator::ExportNativeGeneratedInitCode(FOutputDevice& Out, FOutputDevice& OutDeclarations, FReferenceGatherers& OutReferenceGatherers, const FUnrealSourceFile& SourceFile, FClass* Class, FUHTStringBuilder& OutFriendText) const
+{
+	check(!OutFriendText.Len());
+
+	UE_CLOG(Class->ClassGeneratedBy, LogCompile, Fatal, TEXT("For intrinsic and compiled-in classes, ClassGeneratedBy should always be null"));
+
+	const bool   bIsNoExport  = Class->HasAnyClassFlags(CLASS_NoExport);
+	const bool   bIsDynamic   = FClass::IsDynamic(Class);
+	const FString ClassNameCPP = FNameLookupCPP::GetNameCPP(Class);
+
+	const FString& ApiString = GetAPIString();
+
+	TSet<FName> AlreadyIncludedNames;
+	TArray<UFunction*> FunctionsToExport;
+	bool bAllEditorOnlyFunctions = true;
+	for (UFunction* LocalFunc : TFieldRange<UFunction>(Class,EFieldIteratorFlags::ExcludeSuper))
+	{
+		FName TrueName = FNativeClassHeaderGenerator::GetOverriddenFName(LocalFunc);
+		bool bAlreadyIncluded = false;
+		AlreadyIncludedNames.Add(TrueName, &bAlreadyIncluded);
+		if (bAlreadyIncluded)
+		{
+			// In a dynamic class the same function signature may be used for a Multi- and a Single-cast delegate.
+			if (!LocalFunc->IsA<UDelegateFunction>() || !bIsDynamic)
+			{
+				FError::Throwf(TEXT("The same function linked twice. Function: %s Class: %s"), *LocalFunc->GetName(), *Class->GetName());
+			}
+			continue;
+		}
+		if (!LocalFunc->IsA<UDelegateFunction>())
+		{
+			bAllEditorOnlyFunctions &= LocalFunc->HasAnyFunctionFlags(FUNC_EditorOnly);
+		}
+		FunctionsToExport.Add(LocalFunc);
+	}
+
+	// Sort the list of functions
+	FunctionsToExport.Sort();
+
+	FUHTStringBuilder GeneratedClassRegisterFunctionText;
+
+	// The class itself.
+	{
+		// simple ::StaticClass wrapper to avoid header, link and DLL hell
+		{
+			const FString& SingletonNameNoRegister = GetSingletonName(Class, OutReferenceGatherers.UniqueCrossModuleReferences, false);
+
+			OutDeclarations.Log(FTypeSingletonCache::Get(Class, false).GetExternDecl());
+
+			GeneratedClassRegisterFunctionText.Logf(TEXT("\tUClass* %s\r\n"), *SingletonNameNoRegister);
+			GeneratedClassRegisterFunctionText.Logf(TEXT("\t{\r\n"));
+			GeneratedClassRegisterFunctionText.Logf(TEXT("\t\treturn %s::StaticClass();\r\n"), *ClassNameCPP);
+			GeneratedClassRegisterFunctionText.Logf(TEXT("\t}\r\n"));
+		}
+		const FString& SingletonName = GetSingletonName(Class, OutReferenceGatherers.UniqueCrossModuleReferences);
+
+		FString StaticsStructName = SingletonName.LeftChop(2) + TEXT("_Statics");
+
+		OutFriendText.Logf(TEXT("\tfriend struct %s;\r\n"), *StaticsStructName);
+		OutDeclarations.Log(FTypeSingletonCache::Get(Class).GetExternDecl());
+
+		GeneratedClassRegisterFunctionText.Logf(TEXT("\tstruct %s\r\n"), *StaticsStructName);
+		GeneratedClassRegisterFunctionText.Logf(TEXT("\t{\r\n"));
+
+		FUHTStringBuilder StaticDefinitions;
+
+		FUHTStringBuilder Singletons;
+		UClass* SuperClass = Class->GetSuperClass();
+		if (SuperClass && SuperClass != Class)
+		{
+			OutDeclarations.Log(FTypeSingletonCache::Get(SuperClass).GetExternDecl());
+			Singletons.Logf(TEXT("\t\t(UObject* (*)())%s,\r\n"), *GetSingletonName(SuperClass, OutReferenceGatherers.UniqueCrossModuleReferences).LeftChop(2));
+		}
+		if (!bIsDynamic)
+		{
+			const FString& PackageSingletonName = GetPackageSingletonName(Class->GetOutermost(), OutReferenceGatherers.UniqueCrossModuleReferences);
+
+			OutDeclarations.Logf(TEXT("\t%s_API UPackage* %s;\r\n"), *ApiString, *PackageSingletonName);
+			Singletons.Logf(TEXT("\t\t(UObject* (*)())%s,\r\n"), *PackageSingletonName.LeftChop(2));
+		}
+
+		const TCHAR* SingletonsArray;
+		const TCHAR* SingletonsCount;
+		if (Singletons.Len() != 0)
+		{
+			GeneratedClassRegisterFunctionText.Logf(TEXT("\t\tstatic UObject* (*const DependentSingletons[])();\r\n"));
+
+			StaticDefinitions.Logf(TEXT("\tUObject* (*const %s::DependentSingletons[])() = {\r\n"), *StaticsStructName);
+			StaticDefinitions.Log (*Singletons);
+			StaticDefinitions.Logf(TEXT("\t};\r\n"));
+
+			SingletonsArray = TEXT("DependentSingletons");
+			SingletonsCount = TEXT("UE_ARRAY_COUNT(DependentSingletons)");
+		}
+		else
+		{
+			SingletonsArray = TEXT("nullptr");
+			SingletonsCount = TEXT("0");
+		}
+
+		const TCHAR* FunctionsArray;
+		const TCHAR* FunctionsCount;
+		if (FunctionsToExport.Num() != 0)
+		{
+			GeneratedClassRegisterFunctionText.Log(BEGIN_WRAP_EDITOR_ONLY(bAllEditorOnlyFunctions));
+			GeneratedClassRegisterFunctionText.Log(TEXT("\t\tstatic const FClassFunctionLinkInfo FuncInfo[];\r\n"));
+			GeneratedClassRegisterFunctionText.Log(END_WRAP_EDITOR_ONLY(bAllEditorOnlyFunctions));
+
+			StaticDefinitions.Log(BEGIN_WRAP_EDITOR_ONLY(bAllEditorOnlyFunctions));
+			StaticDefinitions.Logf(TEXT("\tconst FClassFunctionLinkInfo %s::FuncInfo[] = {\r\n"), *StaticsStructName);
+
+			for (UFunction* Function : FunctionsToExport)
+			{
+				const bool bIsEditorOnlyFunction = Function->HasAnyFunctionFlags(FUNC_EditorOnly);
+
+				if (!Function->IsA<UDelegateFunction>())
+				{
+					ExportFunction(Out, OutReferenceGatherers, SourceFile, Function, bIsNoExport);
+				}
+
+				StaticDefinitions.Logf(
+					TEXT("%s\t\t{ &%s, %s },%s\r\n%s"),
+					BEGIN_WRAP_EDITOR_ONLY(bIsEditorOnlyFunction),
+					*GetSingletonNameFuncAddr(Function, OutReferenceGatherers.UniqueCrossModuleReferences),
+					*FNativeClassHeaderGenerator::GetUTF8OverriddenNameForLiteral(Function),
+					*GetGeneratedCodeHashTag(Function),
+					END_WRAP_EDITOR_ONLY(bIsEditorOnlyFunction)
+				);
+			}
+
+			StaticDefinitions.Log(TEXT("\t};\r\n"));
+			StaticDefinitions.Log(END_WRAP_EDITOR_ONLY(bAllEditorOnlyFunctions));
+
+			if (bAllEditorOnlyFunctions)
+			{
+				FunctionsArray = TEXT("IF_WITH_EDITOR(FuncInfo, nullptr)");
+				FunctionsCount = TEXT("IF_WITH_EDITOR(UE_ARRAY_COUNT(FuncInfo), 0)");
+			}
+			else
+			{
+				FunctionsArray = TEXT("FuncInfo");
+				FunctionsCount = TEXT("UE_ARRAY_COUNT(FuncInfo)");
+			}
+		}
+		else
+		{
+			FunctionsArray = TEXT("nullptr");
+			FunctionsCount = TEXT("0");
+		}
+
+		TMap<FName, FString>* MetaDataMap = UMetaData::GetMapForObject(Class);
+		if (MetaDataMap)
+		{
+			FClassMetaData* ClassMetaData = GScriptHelper.FindClassData(Class);
+			if (ClassMetaData && ClassMetaData->bObjectInitializerConstructorDeclared)
+			{
+				MetaDataMap->Add(NAME_ObjectInitializerConstructorDeclared, FString());
+			}
+		}
+
+		FString MetaDataParams = OutputMetaDataCodeForObject(GeneratedClassRegisterFunctionText, StaticDefinitions, Class, *FString::Printf(TEXT("%s::Class_MetaDataParams"), *StaticsStructName), TEXT("\t\t"), TEXT("\t"));
+
+		TArray<FProperty*> Props;
+		for (FProperty* Prop : TFieldRange<FProperty>(Class, EFieldIteratorFlags::ExcludeSuper))
+		{
+			Props.Add(Prop);
+		}
+
+		TTuple<FString, FString> PropertyRange = OutputProperties(GeneratedClassRegisterFunctionText, StaticDefinitions, OutReferenceGatherers, *FString::Printf(TEXT("%s::"), *StaticsStructName), Props, TEXT("\t\t"), TEXT("\t"));
+
+		const TCHAR* InterfaceArray;
+		const TCHAR* InterfaceCount;
+		if (Class->Interfaces.Num() > 0)
+		{
+			GeneratedClassRegisterFunctionText.Log(TEXT("\t\tstatic const UE4CodeGen_Private::FImplementedInterfaceParams InterfaceParams[];\r\n"));
+
+			StaticDefinitions.Logf(TEXT("\t\tconst UE4CodeGen_Private::FImplementedInterfaceParams %s::InterfaceParams[] = {\r\n"), *StaticsStructName);
+			for (const FImplementedInterface& Inter : Class->Interfaces)
+			{
+				check(Inter.Class);
+				FString OffsetString;
+				if (Inter.PointerOffset)
+				{
+					OffsetString = FString::Printf(TEXT("(int32)VTABLE_OFFSET(%s, %s)"), *ClassNameCPP, *FNameLookupCPP::GetNameCPP(Inter.Class, true));
+				}
+				else
+				{
+					OffsetString = TEXT("0");
+				}
+				StaticDefinitions.Logf(
+					TEXT("\t\t\t{ %s, %s, %s },\r\n"),
+					*GetSingletonName(Inter.Class, OutReferenceGatherers.UniqueCrossModuleReferences, false).LeftChop(2),
+					*OffsetString,
+					Inter.bImplementedByK2 ? TEXT("true") : TEXT("false")
+				);
+			}
+			StaticDefinitions.Log(TEXT("\t\t};\r\n"));
+
+			InterfaceArray = TEXT("InterfaceParams");
+			InterfaceCount = TEXT("UE_ARRAY_COUNT(InterfaceParams)");
+		}
+		else
+		{
+			InterfaceArray = TEXT("nullptr");
+			InterfaceCount = TEXT("0");
+		}
+
+		GeneratedClassRegisterFunctionText.Logf(TEXT("\t\tstatic const FCppClassTypeInfoStatic StaticCppClassTypeInfo;\r\n"));
+
+		StaticDefinitions.Logf(TEXT("\tconst FCppClassTypeInfoStatic %s::StaticCppClassTypeInfo = {\r\n"), *StaticsStructName);
+		StaticDefinitions.Logf(TEXT("\t\tTCppClassTypeTraits<%s>::IsAbstract,\r\n"), *FNameLookupCPP::GetNameCPP(Class, Class->HasAllClassFlags(CLASS_Interface)));
+		StaticDefinitions.Logf(TEXT("\t};\r\n"));
+
+		GeneratedClassRegisterFunctionText.Log (TEXT("\t\tstatic const UE4CodeGen_Private::FClassParams ClassParams;\r\n"));
+
+		uint32 ClassFlags = (uint32)Class->ClassFlags;
+		if (!bIsNoExport)
+		{
+			ClassFlags = ClassFlags | CLASS_MatchedSerializers;
+		}
+		ClassFlags = ClassFlags & CLASS_SaveInCompiledInClasses;
+
+		StaticDefinitions.Logf(TEXT("\tconst UE4CodeGen_Private::FClassParams %s::ClassParams = {\r\n"), *StaticsStructName);
+		StaticDefinitions.Logf(TEXT("\t\t&%s::StaticClass,\r\n"), *ClassNameCPP);
+		StaticDefinitions.Logf(TEXT("\t\t%s,\r\n"), (Class->ClassConfigName != NAME_None) ? *CreateUTF8LiteralString(Class->ClassConfigName.ToString()) : TEXT("nullptr"));
+		StaticDefinitions.Log (TEXT("\t\t&StaticCppClassTypeInfo,\r\n"));
+		StaticDefinitions.Logf(TEXT("\t\t%s,\r\n"), SingletonsArray);
+		StaticDefinitions.Logf(TEXT("\t\t%s,\r\n"), FunctionsArray);
+		StaticDefinitions.Logf(TEXT("\t\t%s,\r\n"), *PropertyRange.Get<0>());
+		StaticDefinitions.Logf(TEXT("\t\t%s,\r\n"), InterfaceArray);
+		StaticDefinitions.Logf(TEXT("\t\t%s,\r\n"), SingletonsCount);
+		StaticDefinitions.Logf(TEXT("\t\t%s,\r\n"), FunctionsCount);
+		StaticDefinitions.Logf(TEXT("\t\t%s,\r\n"), *PropertyRange.Get<1>());
+		StaticDefinitions.Logf(TEXT("\t\t%s,\r\n"), InterfaceCount);
+		StaticDefinitions.Logf(TEXT("\t\t0x%08Xu,\r\n"), ClassFlags);
+		StaticDefinitions.Logf(TEXT("\t\t%s\r\n"), *MetaDataParams);
+		StaticDefinitions.Log (TEXT("\t};\r\n"));
+
+		GeneratedClassRegisterFunctionText.Logf(TEXT("\t};\r\n"));
+		GeneratedClassRegisterFunctionText.Log(*StaticDefinitions);
+
+		GeneratedClassRegisterFunctionText.Logf(TEXT("\tUClass* %s\r\n"), *SingletonName);
+		GeneratedClassRegisterFunctionText.Logf(TEXT("\t{\r\n"));
+		if (!bIsDynamic)
+		{
+			GeneratedClassRegisterFunctionText.Logf(TEXT("\t\tstatic UClass* OuterClass = nullptr;\r\n"));
+			GeneratedClassRegisterFunctionText.Logf(TEXT("\t\tif (!OuterClass)\r\n"));
+		}
+		else
+		{
+			const FString& DynamicClassPackageName = FClass::GetTypePackageName(Class);
+			GeneratedClassRegisterFunctionText.Logf(TEXT("\t\tUPackage* OuterPackage = FindOrConstructDynamicTypePackage(TEXT(\"%s\"));\r\n"), *DynamicClassPackageName);
+			GeneratedClassRegisterFunctionText.Logf(TEXT("\t\tUClass* OuterClass = Cast<UClass>(StaticFindObjectFast(UClass::StaticClass(), OuterPackage, TEXT(\"%s\")));\r\n"), *FNativeClassHeaderGenerator::GetOverriddenName(Class));
+			GeneratedClassRegisterFunctionText.Logf(TEXT("\t\tif (!OuterClass || !(OuterClass->ClassFlags & CLASS_Constructed))\r\n"));
+		}
+
+		GeneratedClassRegisterFunctionText.Logf(TEXT("\t\t{\r\n"));
+		GeneratedClassRegisterFunctionText.Logf(TEXT("\t\t\tUE4CodeGen_Private::ConstructUClass(OuterClass, %s::ClassParams);\r\n"), *StaticsStructName);
+
+		TArray<FString> SparseClassDataTypes;
+		((FClass*)Class)->GetSparseClassDataTypes(SparseClassDataTypes);
+		
+		for (const FString& SparseClassDataString : SparseClassDataTypes)
+		{
+			GeneratedClassRegisterFunctionText.Logf(TEXT("\t\t\tOuterClass->SetSparseClassDataStruct(F%s::StaticStruct());\r\n"), *SparseClassDataString);
+		}
+
+		if (bIsDynamic)
+		{
+			FString* CustomDynamicClassInitializationMD = MetaDataMap ? MetaDataMap->Find(TEXT("CustomDynamicClassInitialization")) : nullptr;
+			if (CustomDynamicClassInitializationMD)
+			{
+				GeneratedClassRegisterFunctionText.Logf(TEXT("\t\t\t\t%s(CastChecked<UDynamicClass>(OuterClass));\n"), *(*CustomDynamicClassInitializationMD));
+			}
+		}
+
+		GeneratedClassRegisterFunctionText.Logf(TEXT("\t\t}\r\n"));
+		GeneratedClassRegisterFunctionText.Logf(TEXT("\t\treturn OuterClass;\r\n"));
+		GeneratedClassRegisterFunctionText.Logf(TEXT("\t}\r\n"));
+
+		Out.Logf(TEXT("%s"), *GeneratedClassRegisterFunctionText);
+	}
+
+	if (OutFriendText.Len() && bIsNoExport)
+	{
+		Out.Logf(TEXT("\t/* friend declarations for pasting into noexport class %s\r\n"), *ClassNameCPP);
+		Out.Log(OutFriendText);
+		Out.Logf(TEXT("\t*/\r\n"));
+		OutFriendText.Reset();
+	}
+
+	FString SingletonName = GetSingletonName(Class, OutReferenceGatherers.UniqueCrossModuleReferences);
+	SingletonName.ReplaceInline(TEXT("()"), TEXT(""), ESearchCase::CaseSensitive); // function address
+
+	FString OverriddenClassName = *FNativeClassHeaderGenerator::GetOverriddenName(Class);
+
+	const FString EmptyString;
+	const FString& InitSearchableValuesFunctionName = bIsDynamic ? Class->GetMetaData(NAME_InitializeStaticSearchableValues) : EmptyString;
+	const FString InitSearchableValuesFunctionParam = InitSearchableValuesFunctionName.IsEmpty() ? FString(TEXT("nullptr")) :
+		FString::Printf(TEXT("&%s::%s"), *ClassNameCPP, *InitSearchableValuesFunctionName);
+
+	// Append base class' hash at the end of the generated code, this will force update derived classes
+	// when base class changes during hot-reload.
+	uint32 BaseClassHash = 0;
+	FClass* SuperClass = Class->GetSuperClass();
+	if (SuperClass && !SuperClass->HasAnyClassFlags(CLASS_Intrinsic))
+	{
+		// Since we are dependent on our SuperClass having generated its hash, if it is not available
+		// we will need to wait on it becoming available. Since the SourceFile array provided to the
+		// ParallelFor is in dependency order and does not allow cyclic dependencies, we can be certain
+		// that another thread has started processing the file containing our SuperClass before this
+		// file would have been assigned out,  so we just have to wait
+		while (1)
+		{
+			{
+				FRWScopeLock Lock(GGeneratedCodeHashesLock, SLT_ReadOnly);
+				if (const uint32* Hash = GGeneratedCodeHashes.Find(SuperClass))
+				{
+					BaseClassHash = *Hash;
+					break;
+				}
+			}
+			FPlatformProcess::Sleep(0.01);
+		}
+	}
+	GeneratedClassRegisterFunctionText.Logf(TEXT("\r\n// %u\r\n"), BaseClassHash);
+
+	// Append info for the sparse class data struct onto the text to be hashed
+	TArray<FString> SparseClassDataTypes;
+	((FClass*)Class)->GetSparseClassDataTypes(SparseClassDataTypes);
+
+	for (const FString& SparseClassDataString : SparseClassDataTypes)
+	{
+		UScriptStruct* SparseClassDataStruct = FindObjectSafe<UScriptStruct>(ANY_PACKAGE, *SparseClassDataString);
+		if (!SparseClassDataStruct)
+		{
+			continue;
+		}
+		GeneratedClassRegisterFunctionText.Logf(TEXT("%s\r\n"), *SparseClassDataStruct->GetName());
+		for (FProperty* Child : TFieldRange<FProperty>(SparseClassDataStruct))
+		{
+			GeneratedClassRegisterFunctionText.Logf(TEXT("%s %s\r\n"), *Child->GetCPPType(), *Child->GetNameCPP());
+		}
+	}
+
+	// Calculate generated class initialization code hash so that we know when it changes after hot-reload
+	uint32 ClassHash = GenerateTextHash(*GeneratedClassRegisterFunctionText);
+	AddGeneratedCodeHash(Class, ClassHash);
+	// Emit the IMPLEMENT_CLASS macro to go in the generated cpp file.
+	if (!bIsDynamic)
+	{
+		Out.Logf(TEXT("\tIMPLEMENT_CLASS(%s, %u);\r\n"), *ClassNameCPP, ClassHash);
+	}
+	else
+	{
+		Out.Logf(TEXT("\tIMPLEMENT_DYNAMIC_CLASS(%s, TEXT(\"%s\"), %u);\r\n"), *ClassNameCPP, *OverriddenClassName, ClassHash);
+	}
+
+	Out.Logf(TEXT("\ttemplate<> %sUClass* StaticClass<%s>()\r\n"), *GetAPIString(), *ClassNameCPP);
+	Out.Logf(TEXT("\t{\r\n"));
+	Out.Logf(TEXT("\t\treturn %s::StaticClass();\r\n"), *ClassNameCPP);
+	Out.Logf(TEXT("\t}\r\n"));
+
+	if (bIsDynamic)
+	{
+		const FString& ClassPackageName = FClass::GetTypePackageName(Class);
+		Out.Logf(TEXT("\tstatic FCompiledInDefer Z_CompiledInDefer_UClass_%s(%s, &%s::StaticClass, TEXT(\"%s\"), TEXT(\"%s\"), true, %s, %s, %s);\r\n"),
+			*ClassNameCPP,
+			*SingletonName,
+			*ClassNameCPP,
+			*ClassPackageName,
+			*OverriddenClassName,
+			*AsTEXT(ClassPackageName),
+			*AsTEXT(FNativeClassHeaderGenerator::GetOverriddenPathName(Class)),
+			*InitSearchableValuesFunctionParam);
+	}
+	else
+	{
+		Out.Logf(TEXT("\tstatic FCompiledInDefer Z_CompiledInDefer_UClass_%s(%s, &%s::StaticClass, TEXT(\"%s\"), TEXT(\"%s\"), false, nullptr, nullptr, %s);\r\n"),
+			*ClassNameCPP,
+			*SingletonName,
+			*ClassNameCPP,
+			*Class->GetOutermost()->GetName(),
+			*ClassNameCPP,
+			*InitSearchableValuesFunctionParam);
+	}
+
+
+	if (ClassHasReplicatedProperties(Class))
+	{
+		Out.Logf(TEXT(
+			"\r\n"
+			"\tvoid %s::ValidateGeneratedRepEnums(const TArray<struct FRepRecord>& ClassReps) const\r\n"
+			"\t{\r\n"
+		), *ClassNameCPP);
+
+		FUHTStringBuilder NameBuilder;
+
+		FUHTStringBuilder ValidationBuilder;
+		ValidationBuilder.Log(TEXT("\t\tconst bool bIsValid = true"));
+
+		for (int32 i = Class->FirstOwnedClassRep; i < Class->ClassReps.Num(); ++i)
+		{
+			const FProperty* const Property = Class->ClassReps[i].Property;
+			const FString PropertyName = Property->GetName();
+
+			NameBuilder.Logf(TEXT("\t\tstatic const FName Name_%s(TEXT(\"%s\"));\r\n"), *PropertyName, *FNativeClassHeaderGenerator::GetOverriddenName(Property));
+
+			if (Property->ArrayDim == 1)
+			{
+				ValidationBuilder.Logf(TEXT("\r\n\t\t\t&& Name_%s == ClassReps[(int32)ENetFields_Private::%s].Property->GetFName()"), *PropertyName, *PropertyName);
+			}
+			else
+			{
+				ValidationBuilder.Logf(TEXT("\r\n\t\t\t&& Name_%s == ClassReps[(int32)ENetFields_Private::%s_STATIC_ARRAY].Property->GetFName()"), *PropertyName, *PropertyName);
+			}
+		}
+
+		ValidationBuilder.Log(TEXT(";\r\n"));
+
+		Out.Logf(TEXT(
+			"%s\r\n" // NameBuilder
+			"%s\r\n" // ValidationBuilder
+			"\t\tcheckf(bIsValid, TEXT(\"UHT Generated Rep Indices do not match runtime populated Rep Indices for properties in %s\"));\r\n"
+			"\t}\r\n"
+		), *NameBuilder, *ValidationBuilder, *ClassNameCPP);
+	}
+}
+
+void FNativeClassHeaderGenerator::ExportFunction(FOutputDevice& Out, FReferenceGatherers& OutReferenceGatherers, const FUnrealSourceFile& SourceFile, UFunction* Function, bool bIsNoExport) const
+{
+	UFunction* SuperFunction = Function->GetSuperFunction();
+
+	const bool bIsEditorOnlyFunction = Function->HasAnyFunctionFlags(FUNC_EditorOnly);
+
+	bool bIsDelegate = Function->HasAnyFunctionFlags(FUNC_Delegate);
+
+	const FString& SingletonName = GetSingletonName(Function, OutReferenceGatherers.UniqueCrossModuleReferences);
+	FString StaticsStructName = SingletonName.LeftChop(2) + TEXT("_Statics");
+
+	FUHTStringBuilder CurrentFunctionText;
+	FUHTStringBuilder StaticDefinitions;
+
+	// Begin wrapping editor only functions.  Note: This should always be the first step!
+	if (bIsEditorOnlyFunction)
+	{
+		CurrentFunctionText.Logf(BeginEditorOnlyGuard);
+	}
+
+	CurrentFunctionText.Logf(TEXT("\tstruct %s\r\n"), *StaticsStructName);
+	CurrentFunctionText.Log (TEXT("\t{\r\n"));
+
+	if (bIsNoExport || !(Function->FunctionFlags&FUNC_Event))  // non-events do not export a params struct, so lets do that locally for offset determination
+	{
+		TArray<UScriptStruct*> Structs = FindNoExportStructs(Function);
+		for (UScriptStruct* Struct : Structs)
+		{
+			ExportMirrorsForNoexportStruct(CurrentFunctionText, Struct, /*Indent=*/ 2);
+		}
+
+		ExportEventParm(CurrentFunctionText, OutReferenceGatherers.ForwardDeclarations, Function, /*Indent=*/ 2, /*bOutputConstructor=*/ false, EExportingState::TypeEraseDelegates);
+	}
+
+	UField* FieldOuter = Cast<UField>(Function->GetOuter());
+	const bool bIsDynamic = (FieldOuter && FClass::IsDynamic(FieldOuter));
+
+	FString OuterFunc;
+	if (UObject* Outer = Function->GetOuter())
+	{
+		OuterFunc = Outer->IsA<UPackage>() ? GetPackageSingletonName((UPackage*)Outer, OutReferenceGatherers.UniqueCrossModuleReferences).LeftChop(2) : GetSingletonNameFuncAddr(Function->GetOwnerClass(), OutReferenceGatherers.UniqueCrossModuleReferences);
+	}
+	else
+	{
+		OuterFunc = TEXT("nullptr");
+	}
+	
+	TArray<FProperty*> Props;
+	Algo::Copy(TFieldRange<FProperty>(Function, EFieldIteratorFlags::ExcludeSuper), Props, Algo::NoRef);
+
+	FString StructureSize;
+	if (Props.Num())
+	{
+		UFunction* TempFunction = Function;
+		while (TempFunction->GetSuperFunction())
+		{
+			TempFunction = TempFunction->GetSuperFunction();
+		}
+		FString FunctionName = TempFunction->GetName();
+		if (TempFunction->HasAnyFunctionFlags(FUNC_Delegate))
+		{
+			FunctionName.LeftChopInline(HEADER_GENERATED_DELEGATE_SIGNATURE_SUFFIX_LENGTH, false);
+		}
+
+		StructureSize = FString::Printf(TEXT("sizeof(%s)"), *GetEventStructParamsName(TempFunction->GetOuter(), *FunctionName));
+	}
+	else
+	{
+		StructureSize = TEXT("0");
+	}
+
+	USparseDelegateFunction* SparseDelegateFunction = Cast<USparseDelegateFunction>(Function);
+	const TCHAR* UFunctionObjectFlags = FClass::IsOwnedByDynamicType(Function) ? TEXT("RF_Public|RF_Transient") : TEXT("RF_Public|RF_Transient|RF_MarkAsNative");
+
+	TTuple<FString, FString> PropertyRange = OutputProperties(CurrentFunctionText, StaticDefinitions, OutReferenceGatherers, *FString::Printf(TEXT("%s::"), *StaticsStructName), Props, TEXT("\t\t"), TEXT("\t"));
+
+	const FFunctionData* CompilerInfo = FFunctionData::FindForFunction(Function);
+	const FFuncInfo&     FunctionData = CompilerInfo->GetFunctionData();
+	const bool           bIsNet       = !!(FunctionData.FunctionFlags & (FUNC_NetRequest | FUNC_NetResponse));
+
+	FString MetaDataParams = OutputMetaDataCodeForObject(CurrentFunctionText, StaticDefinitions, Function, *FString::Printf(TEXT("%s::Function_MetaDataParams"), *StaticsStructName), TEXT("\t\t"), TEXT("\t"));
+
+	CurrentFunctionText.Log(TEXT("\t\tstatic const UE4CodeGen_Private::FFunctionParams FuncParams;\r\n"));
+
+	StaticDefinitions.Logf(
+		TEXT("\tconst UE4CodeGen_Private::FFunctionParams %s::FuncParams = { (UObject*(*)())%s, %s, %s, %s, %s, %s, %s, %s, %s, (EFunctionFlags)0x%08X, %d, %d, %s };\r\n"),
+		*StaticsStructName,
+		*OuterFunc,
+		*GetSingletonNameFuncAddr(SuperFunction, OutReferenceGatherers.UniqueCrossModuleReferences),
+		*CreateUTF8LiteralString(FNativeClassHeaderGenerator::GetOverriddenName(Function)),
+		(SparseDelegateFunction ? *CreateUTF8LiteralString(SparseDelegateFunction->OwningClassName.ToString()) : TEXT("nullptr")),
+		(SparseDelegateFunction ? *CreateUTF8LiteralString(SparseDelegateFunction->DelegateName.ToString()) : TEXT("nullptr")),
+		*StructureSize,
+		*PropertyRange.Get<0>(),
+		*PropertyRange.Get<1>(),
+		UFunctionObjectFlags,
+		(uint32)Function->FunctionFlags,
+		bIsNet ? FunctionData.RPCId : 0,
+		bIsNet ? FunctionData.RPCResponseId : 0,
+		*MetaDataParams
+	);
+
+	CurrentFunctionText.Log(TEXT("\t};\r\n"));
+	CurrentFunctionText.Log(*StaticDefinitions);
+
+	CurrentFunctionText.Logf(TEXT("\tUFunction* %s\r\n"), *SingletonName);
+	CurrentFunctionText.Log (TEXT("\t{\r\n"));
+
+	if (!bIsDynamic)
+	{
+		CurrentFunctionText.Logf(TEXT("\t\tstatic UFunction* ReturnFunction = nullptr;\r\n"));
+	}
+	else
+	{
+		FString FunctionName = FNativeClassHeaderGenerator::GetOverriddenNameForLiteral(Function);
+		CurrentFunctionText.Logf(TEXT("\t\tUObject* Outer = %s();\r\n"), *OuterFunc);
+		CurrentFunctionText.Logf(TEXT("\t\tUFunction* ReturnFunction = static_cast<UFunction*>(StaticFindObjectFast( UFunction::StaticClass(), Outer, %s ));\r\n"), *FunctionName);
+	}
+
+	CurrentFunctionText.Logf(TEXT("\t\tif (!ReturnFunction)\r\n"));
+	CurrentFunctionText.Logf(TEXT("\t\t{\r\n"));
+	CurrentFunctionText.Logf(TEXT("\t\t\tUE4CodeGen_Private::ConstructUFunction(ReturnFunction, %s::FuncParams);\r\n"), *StaticsStructName);
+	CurrentFunctionText.Log (TEXT("\t\t}\r\n"));
+	CurrentFunctionText.Log (TEXT("\t\treturn ReturnFunction;\r\n"));
+	CurrentFunctionText.Log (TEXT("\t}\r\n"));
+
+	// End wrapping editor only functions.  Note: This should always be the last step!
+	if (bIsEditorOnlyFunction)
+	{
+		CurrentFunctionText.Logf(EndEditorOnlyGuard);
+	}
+
+	uint32 FunctionHash = GenerateTextHash(*CurrentFunctionText);
+	AddGeneratedCodeHash(Function, FunctionHash);
+	Out.Log(CurrentFunctionText);
+}
+
+void FNativeClassHeaderGenerator::ExportNatives(FOutputDevice& Out, FClass* Class)
+{
+	const FString ClassCPPName = FNameLookupCPP::GetNameCPP(Class);
+	FString TypeName = Class->HasAnyClassFlags(CLASS_Interface) ? FString::Printf(TEXT("I%s"), *Class->GetName()) : ClassCPPName;
+
+	Out.Logf(TEXT("\tvoid %s::StaticRegisterNatives%s()\r\n"), *ClassCPPName, *ClassCPPName);
+	Out.Log(TEXT("\t{\r\n"));
+
+	{
+		bool bAllEditorOnly = true;
+
+		TArray<TTuple<UFunction*, FString>> NamedFunctionsToExport;
+		for (UFunction* Function : TFieldRange<UFunction>(Class, EFieldIteratorFlags::ExcludeSuper))
+		{
+			if ((Function->FunctionFlags & (FUNC_Native | FUNC_NetRequest)) == FUNC_Native)
+			{
+				FString OverriddenName = FNativeClassHeaderGenerator::GetUTF8OverriddenNameForLiteral(Function);
+				NamedFunctionsToExport.Emplace(Function, MoveTemp(OverriddenName));
+
+				if (!Function->HasAnyFunctionFlags(FUNC_EditorOnly))
+				{
+					bAllEditorOnly = false;
+				}
+			}
+		}
+
+		Algo::SortBy(NamedFunctionsToExport, [](const TTuple<UFunction*, FString>& Pair){ return Pair.Get<0>()->GetFName(); }, FNameLexicalLess());
+
+		if (NamedFunctionsToExport.Num() > 0)
+		{
+			FMacroBlockEmitter EditorOnly(Out, TEXT("WITH_EDITOR"));
+			EditorOnly(bAllEditorOnly);
+
+			Out.Logf(TEXT("\t\tUClass* Class = %s::StaticClass();\r\n"), *ClassCPPName);
+			Out.Log(TEXT("\t\tstatic const FNameNativePtrPair Funcs[] = {\r\n"));
+
+			for (const TTuple<UFunction*, FString>& Func : NamedFunctionsToExport)
+			{
+				UFunction* Function = Func.Get<0>();
+
+				EditorOnly(Function->HasAnyFunctionFlags(FUNC_EditorOnly));
+
+				Out.Logf(
+					TEXT("\t\t\t{ %s, &%s::exec%s },\r\n"),
+					*Func.Get<1>(),
+					*TypeName,
+					*Function->GetName()
+				);
+			}
+
+			EditorOnly(bAllEditorOnly);
+
+			Out.Log(TEXT("\t\t};\r\n"));
+			Out.Logf(TEXT("\t\tFNativeFunctionRegistrar::RegisterFunctions(Class, Funcs, UE_ARRAY_COUNT(Funcs));\r\n"));
+		}
+	}
+
+	for (UScriptStruct* Struct : TFieldRange<UScriptStruct>(Class, EFieldIteratorFlags::ExcludeSuper))
+	{
+		if (Struct->StructFlags & STRUCT_Native)
+		{
+			Out.Logf( TEXT("\t\tUScriptStruct::DeferCppStructOps(FName(TEXT(\"%s\")),new UScriptStruct::TCppStructOps<%s%s>);\r\n"), *Struct->GetName(), Struct->GetPrefixCPP(), *Struct->GetName() );
+		}
+	}
+
+	Out.Logf(TEXT("\t}\r\n"));
+}
+
+void FNativeClassHeaderGenerator::ExportInterfaceCallFunctions(FOutputDevice& OutCpp, FUHTStringBuilder& Out, FReferenceGatherers& OutReferenceGatherers, const TArray<UFunction*>& CallbackFunctions, const TCHAR* ClassName) const
+{
+	const FString& APIString = GetAPIString();
+
+	for (UFunction* Function : CallbackFunctions)
+	{
+		FString FunctionName = Function->GetName();
+
+		FFunctionData* CompilerInfo = FFunctionData::FindForFunction(Function);
+
+		const FFuncInfo& FunctionData = CompilerInfo->GetFunctionData();
+		const TCHAR* ConstQualifier = FunctionData.FunctionReference->HasAllFunctionFlags(FUNC_Const) ? TEXT("const ") : TEXT("");
+		FString ExtraParam = FString::Printf(TEXT("%sUObject* O"), ConstQualifier);
+
+		ExportNativeFunctionHeader(Out, OutReferenceGatherers.ForwardDeclarations, FunctionData, EExportFunctionType::Interface, EExportFunctionHeaderStyle::Declaration, *ExtraParam, *APIString);
+		Out.Logf( TEXT(";") LINE_TERMINATOR );
+
+		FString FunctionNameName = FString::Printf(TEXT("NAME_%s_%s"), *FNameLookupCPP::GetNameCPP(CastChecked<UStruct>(Function->GetOuter())), *FunctionName);
+		OutCpp.Logf(TEXT("\tstatic FName %s = FName(TEXT(\"%s\"));") LINE_TERMINATOR, *FunctionNameName, *GetOverriddenFName(Function).ToString());
+
+		ExportNativeFunctionHeader(OutCpp, OutReferenceGatherers.ForwardDeclarations, FunctionData, EExportFunctionType::Interface, EExportFunctionHeaderStyle::Definition, *ExtraParam, *APIString);
+		OutCpp.Logf( LINE_TERMINATOR TEXT("\t{") LINE_TERMINATOR );
+
+		OutCpp.Logf(TEXT("\t\tcheck(O != NULL);") LINE_TERMINATOR);
+		OutCpp.Logf(TEXT("\t\tcheck(O->GetClass()->ImplementsInterface(U%s::StaticClass()));") LINE_TERMINATOR, ClassName);
+
+		FParmsAndReturnProperties Parameters = GetFunctionParmsAndReturn(FunctionData.FunctionReference);
+
+		// See if we need to create Parms struct
+		const bool bHasParms = Parameters.HasParms();
+		if (bHasParms)
+		{
+			FString EventParmStructName = GetEventStructParamsName(Function->GetOuter(), *FunctionName);
+			OutCpp.Logf(TEXT("\t\t%s Parms;") LINE_TERMINATOR, *EventParmStructName);
+		}
+
+		OutCpp.Logf(TEXT("\t\tUFunction* const Func = O->FindFunction(%s);") LINE_TERMINATOR, *FunctionNameName);
+		OutCpp.Log(TEXT("\t\tif (Func)") LINE_TERMINATOR);
+		OutCpp.Log(TEXT("\t\t{") LINE_TERMINATOR);
+
+		// code to populate Parms struct
+		for (FProperty* Param : Parameters.Parms)
+		{
+			const FString ParamName = Param->GetName();
+			OutCpp.Logf(TEXT("\t\t\tParms.%s=%s;") LINE_TERMINATOR, *ParamName, *ParamName);
+		}
+
+		const FString ObjectRef = FunctionData.FunctionReference->HasAllFunctionFlags(FUNC_Const) ? FString::Printf(TEXT("const_cast<UObject*>(O)")) : TEXT("O");
+		OutCpp.Logf(TEXT("\t\t\t%s->ProcessEvent(Func, %s);") LINE_TERMINATOR, *ObjectRef, bHasParms ? TEXT("&Parms") : TEXT("NULL"));
+
+		for (FProperty* Param : Parameters.Parms)
+		{
+			if( Param->HasAllPropertyFlags(CPF_OutParm) && !Param->HasAnyPropertyFlags(CPF_ConstParm|CPF_ReturnParm))
+			{
+				const FString ParamName = Param->GetName();
+				OutCpp.Logf(TEXT("\t\t\t%s=Parms.%s;") LINE_TERMINATOR, *ParamName, *ParamName);
+			}
+		}
+
+		OutCpp.Log(TEXT("\t\t}") LINE_TERMINATOR);
+
+
+		// else clause to call back into native if it's a BlueprintNativeEvent
+		if (Function->FunctionFlags & FUNC_Native)
+		{
+			OutCpp.Logf(TEXT("\t\telse if (auto I = (%sI%s*)(O->GetNativeInterfaceAddress(U%s::StaticClass())))") LINE_TERMINATOR, ConstQualifier, ClassName, ClassName);
+			OutCpp.Log(TEXT("\t\t{") LINE_TERMINATOR);
+
+			OutCpp.Log(TEXT("\t\t\t"));
+			if (Parameters.Return)
+			{
+				OutCpp.Log(TEXT("Parms.ReturnValue = "));
+			}
+
+			OutCpp.Logf(TEXT("I->%s_Implementation("), *FunctionName);
+
+			bool bFirst = true;
+			for (FProperty* Param : Parameters.Parms)
+			{
+				if (!bFirst)
+				{
+					OutCpp.Logf(TEXT(","));
+				}
+				bFirst = false;
+
+				OutCpp.Logf(TEXT("%s"), *Param->GetName());
+			}
+
+			OutCpp.Logf(TEXT(");") LINE_TERMINATOR);
+
+			OutCpp.Logf(TEXT("\t\t}") LINE_TERMINATOR);
+		}
+
+		if (Parameters.Return)
+		{
+			OutCpp.Logf(TEXT("\t\treturn Parms.ReturnValue;") LINE_TERMINATOR);
+		}
+
+		OutCpp.Logf(TEXT("\t}") LINE_TERMINATOR);
+	}
+}
+
+/**
+ * Gets preprocessor string to emit GENERATED_U*_BODY() macro is deprecated.
+ *
+ * @param MacroName Name of the macro to be deprecated.
+ *
+ * @returns Preprocessor string to emit the message.
+ */
+FString GetGeneratedMacroDeprecationWarning(const TCHAR* MacroName)
+{
+	// Deprecation warning is disabled right now. After people get familiar with the new macro it should be re-enabled.
+	//return FString() + TEXT("EMIT_DEPRECATED_WARNING_MESSAGE(\"") + MacroName + TEXT("() macro is deprecated. Please use GENERATED_BODY() macro instead.\")") LINE_TERMINATOR;
+	return TEXT("");
+}
+
+/**
+ * Returns a string with access specifier that was met before parsing GENERATED_BODY() macro to preserve it.
+ *
+ * @param Class Class for which to return the access specifier.
+ *
+ * @returns Access specifier string.
+ */
+FString GetPreservedAccessSpecifierString(FClass* Class)
+{
+	FString PreservedAccessSpecifier;
+	if (FClassMetaData* Data = GScriptHelper.FindClassData(Class))
+	{
+		switch (Data->GeneratedBodyMacroAccessSpecifier)
+		{
+		case EAccessSpecifier::ACCESS_Private:
+			PreservedAccessSpecifier = "private:";
+			break;
+		case EAccessSpecifier::ACCESS_Protected:
+			PreservedAccessSpecifier = "protected:";
+			break;
+		case EAccessSpecifier::ACCESS_Public:
+			PreservedAccessSpecifier = "public:";
+			break;
+		case EAccessSpecifier::ACCESS_NotAnAccessSpecifier :
+			PreservedAccessSpecifier = FString::Printf(TEXT("static_assert(false, \"Unknown access specifier for GENERATED_BODY() macro in class %s.\");"), *GetNameSafe(Class));
+			break;
+		}
+	}
+
+	return PreservedAccessSpecifier + LINE_TERMINATOR;
+}
+
+void WriteMacro(FOutputDevice& Output, const FString& MacroName, FString MacroContent)
+{
+	Output.Log(Macroize(*MacroName, MoveTemp(MacroContent)));
+}
+
+static FString PrivatePropertiesOffsetGetters(const UStruct* Struct, const FString& StructCppName)
+{
+	check(Struct);
+
+	FUHTStringBuilder Result;
+	for (const FProperty* Property : TFieldRange<FProperty>(Struct, EFieldIteratorFlags::ExcludeSuper))
+	{
+		if (Property && Property->HasAnyPropertyFlags(CPF_NativeAccessSpecifierPrivate | CPF_NativeAccessSpecifierProtected) && !Property->HasAnyPropertyFlags(CPF_EditorOnly))
+		{
+			const FBoolProperty* BoolProperty = CastField<const FBoolProperty>(Property);
+			if (BoolProperty && !BoolProperty->IsNativeBool()) // if it's a bitfield
+			{
+				continue;
+			}
+
+			FString PropertyName = Property->GetName();
+			if (Property->HasAllPropertyFlags(CPF_Deprecated))
+			{
+				PropertyName += TEXT("_DEPRECATED");
+			}
+			Result.Logf(TEXT("\tFORCEINLINE static uint32 __PPO__%s() { return STRUCT_OFFSET(%s, %s); }") LINE_TERMINATOR,
+				*PropertyName, *StructCppName, *PropertyName);
+		}
+	}
+
+	return MoveTemp(Result);
+}
+
+void FNativeClassHeaderGenerator::ExportClassFromSourceFileInner(
+	FOutputDevice&           OutGeneratedHeaderText,
+	FOutputDevice&           OutCpp,
+	FOutputDevice&           OutDeclarations,
+	FReferenceGatherers&     OutReferenceGatherers,
+	FClass*                  Class,
+	const FUnrealSourceFile& SourceFile,
+	EExportClassOutFlags&    OutFlags
+) const
+{
+	FUHTStringBuilder StandardUObjectConstructorsMacroCall;
+	FUHTStringBuilder EnhancedUObjectConstructorsMacroCall;
+
+	FClassMetaData* ClassData = GScriptHelper.FindClassData(Class);
+	checkf(ClassData, TEXT("No class data generated for file %s"), *SourceFile.GetFilename());
+
+	// C++ -> VM stubs (native function execs)
+	FUHTStringBuilder ClassMacroCalls;
+	FUHTStringBuilder ClassNoPureDeclsMacroCalls;
+	ExportNativeFunctions(OutGeneratedHeaderText, OutCpp, ClassMacroCalls, ClassNoPureDeclsMacroCalls, OutReferenceGatherers, SourceFile, Class, ClassData);
+
+	// Get Callback functions
+	TArray<UFunction*> CallbackFunctions;
+	for (UFunction* Function : TFieldRange<UFunction>(Class, EFieldIteratorFlags::ExcludeSuper))
+	{
+		if ((Function->FunctionFlags & FUNC_Event) && Function->GetSuperFunction() == nullptr)
+		{
+			CallbackFunctions.Add(Function);
+		}
+	}
+
+	FUHTStringBuilder PrologMacroCalls;
+	if (CallbackFunctions.Num() != 0)
+	{
+		Algo::SortBy(CallbackFunctions, [](UObject* Obj) { return Obj->GetName(); });
+
+		FUHTStringBuilder UClassMacroContent;
+
+		// export parameters structs for all events and delegates
+		for (UFunction* Function : CallbackFunctions)
+		{
+			ExportEventParm(UClassMacroContent, OutReferenceGatherers.ForwardDeclarations, Function, /*Indent=*/ 1, /*bOutputConstructor=*/ true, EExportingState::Normal);
+		}
+
+		FString MacroName = SourceFile.GetGeneratedMacroName(ClassData, TEXT("_EVENT_PARMS"));
+		WriteMacro(OutGeneratedHeaderText, MacroName, UClassMacroContent);
+		PrologMacroCalls.Logf(TEXT("\t%s\r\n"), *MacroName);
+
+		// VM -> C++ proxies (events and delegates).
+		FOutputDeviceNull NullOutput;
+		FOutputDevice& CallbackOut = Class->HasAnyClassFlags(CLASS_NoExport) ? NullOutput : OutCpp;
+		FString CallbackWrappersMacroName = SourceFile.GetGeneratedMacroName(ClassData, TEXT("_CALLBACK_WRAPPERS"));
+		ExportCallbackFunctions(
+			OutGeneratedHeaderText,
+			CallbackOut,
+			OutReferenceGatherers.ForwardDeclarations,
+			CallbackFunctions,
+			*CallbackWrappersMacroName,
+			(Class->ClassFlags & CLASS_Interface) ? EExportCallbackType::Interface : EExportCallbackType::Class,
+			*GetAPIString()
+		);
+
+		ClassMacroCalls.Logf(TEXT("\t%s\r\n"), *CallbackWrappersMacroName);
+		ClassNoPureDeclsMacroCalls.Logf(TEXT("\t%s\r\n"), *CallbackWrappersMacroName);
+	}
+
+	// Class definition.
+	if (!Class->HasAnyClassFlags(CLASS_NoExport))
+	{
+		ExportNatives(OutCpp, Class);
+	}
+
+	FUHTStringBuilder FriendText;
+	ExportNativeGeneratedInitCode(OutCpp, OutDeclarations, OutReferenceGatherers, SourceFile, Class, FriendText);
+
+	FClass* SuperClass = Class->GetSuperClass();
+
+	// the name for the C++ version of the UClass
+	const FString ClassCPPName = FNameLookupCPP::GetNameCPP(Class);
+	const FString SuperClassCPPName = (SuperClass ? FNameLookupCPP::GetNameCPP(SuperClass) : TEXT("None"));
+
+	FString APIArg = API;
+	if (!Class->HasAnyClassFlags(CLASS_MinimalAPI))
+	{
+		APIArg = TEXT("NO");
+	}
+
+	FString PPOMacroName;
+
+	ClassDefinitionRange ClassRange;
+	if (ClassDefinitionRange* FoundRange = ClassDefinitionRanges.Find(Class))
+	{
+		ClassRange = *FoundRange;
+		ClassRange.Validate();
+	}
+
+	FString GeneratedSerializeFunctionCPP;
+	FString GeneratedSerializeFunctionHeaderMacroName;
+
+	if (const FArchiveTypeDefinePair* ArchiveTypeDefinePair = GClassSerializerMap.Find(Class))
+	{
+		FString EnclosingDefines;
+		FUHTStringBuilder Boilerplate, BoilerPlateCPP;
+		const TCHAR* MacroNameHeader;
+		const TCHAR* MacroNameCPP;
+		GeneratedSerializeFunctionHeaderMacroName = SourceFile.GetGeneratedMacroName(ClassData, TEXT("_ARCHIVESERIALIZER"));
+
+		EnclosingDefines = ArchiveTypeDefinePair->EnclosingDefine;
+		if (ArchiveTypeDefinePair->ArchiveType == ESerializerArchiveType::StructuredArchiveRecord)
+		{
+			MacroNameHeader = TEXT("DECLARE_FARCHIVE_SERIALIZER");
+			MacroNameCPP = TEXT("IMPLEMENT_FARCHIVE_SERIALIZER");
+		}
+		else
+		{
+			MacroNameHeader = TEXT("DECLARE_FSTRUCTUREDARCHIVE_SERIALIZER");
+			MacroNameCPP = TEXT("IMPLEMENT_FSTRUCTUREDARCHIVE_SERIALIZER");
+		}
+
+		// if the existing Serialize function was wrapped in a compiler define directive, we need to replicate that on the generated function
+		if (EnclosingDefines.Len())
+		{
+			OutGeneratedHeaderText.Logf(TEXT("#if %s\r\n"), *EnclosingDefines);
+			BoilerPlateCPP.Logf(TEXT("#if %s\r\n"), *EnclosingDefines);
+		}
+
+		Boilerplate.Logf(TEXT("\t%s(%s, %s_API)\r\n"), MacroNameHeader, *ClassCPPName, *APIArg);
+		OutGeneratedHeaderText.Log(Macroize(*GeneratedSerializeFunctionHeaderMacroName, *Boilerplate));
+		BoilerPlateCPP.Logf(TEXT("\t%s(%s)\r\n"), MacroNameCPP, *ClassCPPName);
+
+		if (EnclosingDefines.Len())
+		{
+			OutGeneratedHeaderText.Logf(TEXT("#else\r\n"));
+			OutGeneratedHeaderText.Log(Macroize(*GeneratedSerializeFunctionHeaderMacroName, TEXT("")));
+			OutGeneratedHeaderText.Logf(TEXT("#endif\r\n"));
+			BoilerPlateCPP.Logf(TEXT("#endif\r\n"));
+		}
+
+		GeneratedSerializeFunctionCPP = BoilerPlateCPP;
+	}
+
+	{
+		FUHTStringBuilder Boilerplate;
+
+		// Export the class's native function registration.
+		Boilerplate.Logf(TEXT("private:\r\n"));
+		Boilerplate.Logf(TEXT("\tstatic void StaticRegisterNatives%s();\r\n"), *ClassCPPName);
+		Boilerplate.Log(*FriendText);
+		Boilerplate.Logf(TEXT("public:\r\n"));
+
+		const bool bCastedClass = Class->HasAnyCastFlag(CASTCLASS_AllFlags) && SuperClass && Class->ClassCastFlags != SuperClass->ClassCastFlags;
+
+		Boilerplate.Logf(TEXT("\tDECLARE_CLASS(%s, %s, COMPILED_IN_FLAGS(%s%s), %s, TEXT(\"%s\"), %s_API)\r\n"),
+			*ClassCPPName,
+			*SuperClassCPPName,
+			Class->HasAnyClassFlags(CLASS_Abstract) ? TEXT("CLASS_Abstract") : TEXT("0"),
+			*GetClassFlagExportText(Class),
+			bCastedClass ? *FString::Printf(TEXT("CASTCLASS_%s"), *ClassCPPName) : TEXT("CASTCLASS_None"),
+			*FClass::GetTypePackageName(Class),
+			*APIArg);
+
+		Boilerplate.Logf(TEXT("\tDECLARE_SERIALIZER(%s)\r\n"), *ClassCPPName);
+
+		// Add the serialization function declaration if we generated one
+		if (GeneratedSerializeFunctionHeaderMacroName.Len() > 0)
+		{
+			Boilerplate.Logf(TEXT("\t%s\r\n"), *GeneratedSerializeFunctionHeaderMacroName);
+		}
+
+		if (SuperClass && Class->ClassWithin != SuperClass->ClassWithin)
+		{
+			Boilerplate.Logf(TEXT("\tDECLARE_WITHIN(%s)\r\n"), *FNameLookupCPP::GetNameCPP(Class->GetClassWithin()));
+		}
+
+		if (Class->HasAnyClassFlags(CLASS_Interface))
+		{
+			ExportConstructorsMacros(OutGeneratedHeaderText, OutCpp, StandardUObjectConstructorsMacroCall, EnhancedUObjectConstructorsMacroCall, SourceFile.GetGeneratedMacroName(ClassData), Class, *APIArg);
+
+			FString InterfaceMacroName = SourceFile.GetGeneratedMacroName(ClassData, TEXT("_GENERATED_UINTERFACE_BODY"));
+			OutGeneratedHeaderText.Log(Macroize(*(InterfaceMacroName + TEXT("()")), *Boilerplate));
+
+			int32 ClassGeneratedBodyLine = ClassData->GetGeneratedBodyLine();
+
+			FString DeprecationWarning = GetGeneratedMacroDeprecationWarning(TEXT("GENERATED_UINTERFACE_BODY"));
+
+			const TCHAR* Offset = TEXT("\t");
+
+			OutGeneratedHeaderText.Log(
+				Macroize(
+					*SourceFile.GetGeneratedBodyMacroName(ClassGeneratedBodyLine, true),
+					FString::Printf(TEXT("\t%s\t%s\t%s()") LINE_TERMINATOR TEXT("%s\t%s")
+						, *DeprecationWarning
+						, DisableDeprecationWarnings
+						, *InterfaceMacroName
+						, *StandardUObjectConstructorsMacroCall
+						, EnableDeprecationWarnings
+					)
+				)
+			);
+
+			OutGeneratedHeaderText.Log(
+				Macroize(
+					*SourceFile.GetGeneratedBodyMacroName(ClassGeneratedBodyLine),
+					FString::Printf(TEXT("\t%s\t%s()") LINE_TERMINATOR TEXT("%s%s\t%s")
+						, DisableDeprecationWarnings
+						, *InterfaceMacroName
+						, *EnhancedUObjectConstructorsMacroCall
+						, *GetPreservedAccessSpecifierString(Class)
+						, EnableDeprecationWarnings
+					)
+				)
+			);
+
+			// =============================================
+			// Export the pure interface version of the class
+
+			// the name of the pure interface class
+			FString InterfaceCPPName = FString::Printf(TEXT("I%s"), *Class->GetName());
+			FString SuperInterfaceCPPName;
+			if (SuperClass != NULL)
+			{
+				SuperInterfaceCPPName = FString::Printf(TEXT("I%s"), *SuperClass->GetName());
+			}
+
+			// Thunk functions
+			FUHTStringBuilder InterfaceBoilerplate;
+
+			InterfaceBoilerplate.Logf(TEXT("protected:\r\n\tvirtual ~%s() {}\r\n"), *InterfaceCPPName);
+			InterfaceBoilerplate.Logf(TEXT("public:\r\n\ttypedef %s UClassType;\r\n"), *ClassCPPName);
+			InterfaceBoilerplate.Logf(TEXT("\ttypedef %s ThisClass;\r\n"), *InterfaceCPPName);
+
+			ExportInterfaceCallFunctions(OutCpp, InterfaceBoilerplate, OutReferenceGatherers, CallbackFunctions, *Class->GetName());
+
+			// we'll need a way to get to the UObject portion of a native interface, so that we can safely pass native interfaces
+			// to script VM functions
+			if (SuperClass->IsChildOf(UInterface::StaticClass()))
+			{
+				// Note: This used to be declared as a pure virtual function, but it was changed here in order to allow the Blueprint nativization process
+				// to detect C++ interface classes that explicitly declare pure virtual functions via type traits. This code will no longer trigger that check.
+				InterfaceBoilerplate.Logf(TEXT("\tvirtual UObject* _getUObject() const { check(0 && \"Missing required implementation.\"); return nullptr; }\r\n"));
+			}
+
+			if (ClassHasReplicatedProperties(Class))
+			{
+				WriteReplicatedMacroData(ClassRange, *ClassCPPName, *APIArg, Class, SuperClass, InterfaceBoilerplate, SourceFile, OutFlags);
+			}
+
+			FString NoPureDeclsMacroName = SourceFile.GetGeneratedMacroName(ClassData, TEXT("_INCLASS_IINTERFACE_NO_PURE_DECLS"));
+			WriteMacro(OutGeneratedHeaderText, NoPureDeclsMacroName, InterfaceBoilerplate);
+			ClassNoPureDeclsMacroCalls.Logf(TEXT("\t%s\r\n"), *NoPureDeclsMacroName);
+
+			FString MacroName = SourceFile.GetGeneratedMacroName(ClassData, TEXT("_INCLASS_IINTERFACE"));
+			WriteMacro(OutGeneratedHeaderText, MacroName, InterfaceBoilerplate);
+			ClassMacroCalls.Logf(TEXT("\t%s\r\n"), *MacroName);
+		}
+		else
+		{
+			// export the class's config name
+			if (SuperClass && Class->ClassConfigName != NAME_None && Class->ClassConfigName != SuperClass->ClassConfigName)
+			{
+				Boilerplate.Logf(TEXT("\tstatic const TCHAR* StaticConfigName() {return TEXT(\"%s\");}\r\n\r\n"), *Class->ClassConfigName.ToString());
+			}
+
+			// export implementation of _getUObject for classes that implement interfaces
+			if (Class->Interfaces.Num() > 0)
+			{
+				Boilerplate.Logf(TEXT("\tvirtual UObject* _getUObject() const override { return const_cast<%s*>(this); }\r\n"), *ClassCPPName);
+			}
+
+			if (ClassHasReplicatedProperties(Class))
+			{
+				WriteReplicatedMacroData(ClassRange, *ClassCPPName, *APIArg, Class, SuperClass, Boilerplate, SourceFile, OutFlags);
+			}
+
+			{
+				FString NoPureDeclsMacroName = SourceFile.GetGeneratedMacroName(ClassData, TEXT("_INCLASS_NO_PURE_DECLS"));
+				WriteMacro(OutGeneratedHeaderText, NoPureDeclsMacroName, Boilerplate);
+				ClassNoPureDeclsMacroCalls.Logf(TEXT("\t%s\r\n"), *NoPureDeclsMacroName);
+
+				FString MacroName = SourceFile.GetGeneratedMacroName(ClassData, TEXT("_INCLASS"));
+				WriteMacro(OutGeneratedHeaderText, MacroName, Boilerplate);
+				ClassMacroCalls.Logf(TEXT("\t%s\r\n"), *MacroName);
+
+				ExportConstructorsMacros(OutGeneratedHeaderText, OutCpp, StandardUObjectConstructorsMacroCall, EnhancedUObjectConstructorsMacroCall, SourceFile.GetGeneratedMacroName(ClassData), Class, *APIArg);
+			}
+			{
+				const FString PrivatePropertiesOffsets = PrivatePropertiesOffsetGetters(Class, ClassCPPName);
+				const FString PPOMacroNameRaw = SourceFile.GetGeneratedMacroName(ClassData, TEXT("_PRIVATE_PROPERTY_OFFSET"));
+				PPOMacroName = FString::Printf(TEXT("\t%s\r\n"), *PPOMacroNameRaw);
+				WriteMacro(OutGeneratedHeaderText, PPOMacroNameRaw, PrivatePropertiesOffsets);
+			}
+		}
+	}
+
+	{
+		FString MacroName = SourceFile.GetGeneratedMacroName(ClassData->GetPrologLine(), TEXT("_PROLOG"));
+		WriteMacro(OutGeneratedHeaderText, MacroName, PrologMacroCalls);
+	}
+
+	{
+		const TCHAR* Public = TEXT("public:") LINE_TERMINATOR;
+
+		const bool bIsIInterface = Class->HasAnyClassFlags(CLASS_Interface);
+
+		const TCHAR* MacroName;
+		FString DeprecationWarning;
+		FString LegacyGeneratedBody;
+		FString GeneratedBody;
+		int32 GeneratedBodyLine;
+
+
+		if (bIsIInterface)
+		{
+			MacroName = TEXT("GENERATED_IINTERFACE_BODY()");
+			GeneratedBodyLine = ClassData->GetInterfaceGeneratedBodyLine();
+			LegacyGeneratedBody = ClassMacroCalls;
+			GeneratedBody = ClassNoPureDeclsMacroCalls;
+		}
+		else
+		{
+			MacroName = TEXT("GENERATED_UCLASS_BODY()");
+			DeprecationWarning = GetGeneratedMacroDeprecationWarning(MacroName);
+			GeneratedBodyLine = ClassData->GetGeneratedBodyLine();
+			LegacyGeneratedBody = FString::Printf(TEXT("%s%s%s"), *PPOMacroName, *ClassMacroCalls, *StandardUObjectConstructorsMacroCall);
+			GeneratedBody = FString::Printf(TEXT("%s%s%s"), *PPOMacroName, *ClassNoPureDeclsMacroCalls, *EnhancedUObjectConstructorsMacroCall);
+		}
+
+		FString WrappedLegacyGeneratedBody = FString::Printf(TEXT("%s%s%s%s%s%s"), *DeprecationWarning, DisableDeprecationWarnings, Public, *LegacyGeneratedBody, Public, EnableDeprecationWarnings);
+		FString WrappedGeneratedBody = FString::Printf(TEXT("%s%s%s%s%s"), DisableDeprecationWarnings, Public, *GeneratedBody, *GetPreservedAccessSpecifierString(Class), EnableDeprecationWarnings);
+
+		OutGeneratedHeaderText.Log(Macroize(*SourceFile.GetGeneratedBodyMacroName(GeneratedBodyLine, true), MoveTemp(WrappedLegacyGeneratedBody)));
+		OutGeneratedHeaderText.Log(Macroize(*SourceFile.GetGeneratedBodyMacroName(GeneratedBodyLine, false), MoveTemp(WrappedGeneratedBody)));
+	}
+
+	// Forward declare the StaticClass specialisation in the header
+	OutGeneratedHeaderText.Logf(TEXT("template<> %sUClass* StaticClass<class %s>();\r\n\r\n"), *GetAPIString(), *ClassCPPName);
+
+	// If there is a serialization function implementation for the CPP file, add it now
+	if (GeneratedSerializeFunctionCPP.Len())
+	{
+		OutCpp.Log(GeneratedSerializeFunctionCPP);
+	}
+}
+
+/**
+* Generates private copy-constructor declaration.
+*
+* @param Out Output device to generate to.
+* @param Class Class to generate constructor for.
+* @param API API string for this constructor.
+*/
+void ExportCopyConstructorDefinition(FOutputDevice& Out, const TCHAR* API, const TCHAR* ClassCPPName)
+{
+	Out.Logf(TEXT("private:\r\n"));
+	Out.Logf(TEXT("\t/** Private move- and copy-constructors, should never be used */\r\n"));
+	Out.Logf(TEXT("\t%s_API %s(%s&&);\r\n"), API, ClassCPPName, ClassCPPName);
+	Out.Logf(TEXT("\t%s_API %s(const %s&);\r\n"), API, ClassCPPName, ClassCPPName);
+	Out.Logf(TEXT("public:\r\n"));
+}
+
+/**
+ * Generates vtable helper caller and eventual constructor body.
+ *
+ * @param Out Output device to generate to.
+ * @param Class Class to generate for.
+ * @param API API string.
+ */
+void ExportVTableHelperCtorAndCaller(FOutputDevice& Out, FClassMetaData* ClassData, const TCHAR* API, const TCHAR* ClassCPPName)
+{
+	if (!ClassData->bCustomVTableHelperConstructorDeclared)
+	{
+		Out.Logf(TEXT("\tDECLARE_VTABLE_PTR_HELPER_CTOR(%s_API, %s);" LINE_TERMINATOR), API, ClassCPPName);
+	}
+	Out.Logf(TEXT("DEFINE_VTABLE_PTR_HELPER_CTOR_CALLER(%s);" LINE_TERMINATOR), ClassCPPName);
+}
+
+/**
+ * Generates standard constructor declaration.
+ *
+ * @param Out Output device to generate to.
+ * @param Class Class to generate constructor for.
+ * @param API API string for this constructor.
+ */
+void ExportStandardConstructorsMacro(FOutputDevice& Out, FClass* Class, FClassMetaData* ClassData, const TCHAR* API, const TCHAR* ClassCPPName)
+{
+	if (!Class->HasAnyClassFlags(CLASS_CustomConstructor))
+	{
+		Out.Logf(TEXT("\t/** Standard constructor, called after all reflected properties have been initialized */\r\n"));
+		Out.Logf(TEXT("\t%s_API %s(const FObjectInitializer& ObjectInitializer%s);\r\n"), API, ClassCPPName,
+			ClassData->bDefaultConstructorDeclared ? TEXT("") : TEXT(" = FObjectInitializer::Get()"));
+	}
+	Out.Logf(TEXT("\tDEFINE_DEFAULT_OBJECT_INITIALIZER_CONSTRUCTOR_CALL(%s)\r\n"), ClassCPPName);
+
+	ExportVTableHelperCtorAndCaller(Out, ClassData, API, ClassCPPName);
+	ExportCopyConstructorDefinition(Out, API, ClassCPPName);
+}
+
+/**
+ * Generates constructor definition.
+ *
+ * @param Out Output device to generate to.
+ * @param Class Class to generate constructor for.
+ * @param API API string for this constructor.
+ */
+void ExportConstructorDefinition(FOutputDevice& Out, FClass* Class, FClassMetaData* ClassData, const TCHAR* API, const TCHAR* ClassCPPName)
+{
+	if (!ClassData->bConstructorDeclared)
+	{
+		Out.Logf(TEXT("\t/** Standard constructor, called after all reflected properties have been initialized */\r\n"));
+
+		// Assume super class has OI constructor, this may not always be true but we should always be able to check this.
+		// In any case, it will default to old behaviour before we even checked this.
+		bool bSuperClassObjectInitializerConstructorDeclared = true;
+		FClass* SuperClass = Class->GetSuperClass();
+		if (SuperClass != nullptr)
+		{
+			FClassMetaData* SuperClassData = GScriptHelper.FindClassData(SuperClass);
+			if (SuperClassData)
+			{
+				// Since we are dependent on our SuperClass having determined which constructors are defined, 
+				// if it is not yet determined we will need to wait on it becoming available. 
+				// Since the SourceFile array provided to the ParallelFor is in dependency order and does not allow cyclic dependencies, 
+				// we can be certain that another thread has started processing the file containing our SuperClass before this
+				// file would have been assigned out,  so we just have to wait
+				while (!SuperClassData->bConstructorDeclared)
+				{
+					FPlatformProcess::Sleep(0.01);
+				}
+
+				bSuperClassObjectInitializerConstructorDeclared = SuperClassData->bObjectInitializerConstructorDeclared;
+			}
+		}
+		if (bSuperClassObjectInitializerConstructorDeclared)
+		{
+			Out.Logf(TEXT("\t%s_API %s(const FObjectInitializer& ObjectInitializer = FObjectInitializer::Get()) : Super(ObjectInitializer) { };\r\n"), API, ClassCPPName);
+			ClassData->bObjectInitializerConstructorDeclared = true;
+		}
+		else
+		{
+			Out.Logf(TEXT("\t%s_API %s() { };\r\n"), API, ClassCPPName);
+			ClassData->bDefaultConstructorDeclared = true;
+		}
+
+		ClassData->bConstructorDeclared = true;
+	}
+	ExportCopyConstructorDefinition(Out, API, ClassCPPName);
+}
+
+/**
+ * Generates constructor call definition.
+ *
+ * @param Out Output device to generate to.
+ * @param Class Class to generate constructor call definition for.
+ */
+void ExportDefaultConstructorCallDefinition(FOutputDevice& Out, FClassMetaData* ClassData, const TCHAR* ClassCPPName)
+{
+	if (ClassData->bObjectInitializerConstructorDeclared)
+	{
+		Out.Logf(TEXT("\tDEFINE_DEFAULT_OBJECT_INITIALIZER_CONSTRUCTOR_CALL(%s)\r\n"), ClassCPPName);
+	}
+	else if (ClassData->bDefaultConstructorDeclared)
+	{
+		Out.Logf(TEXT("\tDEFINE_DEFAULT_CONSTRUCTOR_CALL(%s)\r\n"), ClassCPPName);
+	}
+	else
+	{
+		Out.Logf(TEXT("\tDEFINE_FORBIDDEN_DEFAULT_CONSTRUCTOR_CALL(%s)\r\n"), ClassCPPName);
+	}
+}
+
+/**
+ * Generates enhanced constructor declaration.
+ *
+ * @param Out Output device to generate to.
+ * @param Class Class to generate constructor for.
+ * @param API API string for this constructor.
+ */
+void ExportEnhancedConstructorsMacro(FOutputDevice& Out, FClass* Class, FClassMetaData* ClassData, const TCHAR* API, const TCHAR* ClassCPPName)
+{
+	ExportConstructorDefinition(Out, Class, ClassData, API, ClassCPPName);
+	ExportVTableHelperCtorAndCaller(Out, ClassData, API, ClassCPPName);
+	ExportDefaultConstructorCallDefinition(Out, ClassData, ClassCPPName);
+}
+
+/**
+ * Gets a package relative inclusion path of the given source file for build.
+ *
+ * @param SourceFile Given source file.
+ *
+ * @returns Inclusion path.
+ */
+FString GetBuildPath(FUnrealSourceFile& SourceFile)
+{
+	FString Out = SourceFile.GetFilename();
+
+	ConvertToBuildIncludePath(SourceFile.GetPackage(), Out);
+
+	return Out;
+}
+
+void FNativeClassHeaderGenerator::ExportConstructorsMacros(FOutputDevice& OutGeneratedHeaderText, FOutputDevice& Out, FOutputDevice& StandardUObjectConstructorsMacroCall, FOutputDevice& EnhancedUObjectConstructorsMacroCall, const FString& ConstructorsMacroPrefix, FClass* Class, const TCHAR* APIArg)
+{
+	const FString ClassCPPName = FNameLookupCPP::GetNameCPP(Class);
+
+	FClassMetaData* ClassData = GScriptHelper.FindClassData(Class);
+	check(ClassData);
+
+	FUHTStringBuilder StdMacro;
+	FUHTStringBuilder EnhMacro;
+	FString StdMacroName = ConstructorsMacroPrefix + TEXT("_STANDARD_CONSTRUCTORS");
+	FString EnhMacroName = ConstructorsMacroPrefix + TEXT("_ENHANCED_CONSTRUCTORS");
+
+	ExportStandardConstructorsMacro(StdMacro, Class, ClassData, APIArg, *ClassCPPName);
+	ExportEnhancedConstructorsMacro(EnhMacro, Class, ClassData, APIArg, *ClassCPPName);
+
+	if (!ClassData->bCustomVTableHelperConstructorDeclared)
+	{
+		Out.Logf(TEXT("\tDEFINE_VTABLE_PTR_HELPER_CTOR(%s);" LINE_TERMINATOR), *ClassCPPName);
+	}
+
+	OutGeneratedHeaderText.Log(Macroize(*StdMacroName, *StdMacro));
+	OutGeneratedHeaderText.Log(Macroize(*EnhMacroName, *EnhMacro));
+
+	StandardUObjectConstructorsMacroCall.Logf(TEXT("\t%s\r\n"), *StdMacroName);
+	EnhancedUObjectConstructorsMacroCall.Logf(TEXT("\t%s\r\n"), *EnhMacroName);
+}
+
+bool FNativeClassHeaderGenerator::WriteHeader(const FPreloadHeaderFileInfo& FileInfo, const FString& InBodyText, const TSet<FString>& InAdditionalHeaders, FReferenceGatherers& InOutReferenceGatherers, FGraphEventRef& OutSaveTempTask) const
+{
+	FUHTStringBuilder GeneratedHeaderTextWithCopyright;
+	GeneratedHeaderTextWithCopyright.Logf(TEXT("%s"), HeaderCopyright);
+	GeneratedHeaderTextWithCopyright.Log(TEXT("#include \"UObject/ObjectMacros.h\"\r\n"));
+	GeneratedHeaderTextWithCopyright.Log(TEXT("#include \"UObject/ScriptMacros.h\"\r\n"));
+
+	for (const FString& AdditionalHeader : InAdditionalHeaders)
+	{
+		GeneratedHeaderTextWithCopyright.Logf(TEXT("#include \"%s\"\r\n"), *AdditionalHeader);
+	}
+
+	GeneratedHeaderTextWithCopyright.Log(LINE_TERMINATOR);
+	GeneratedHeaderTextWithCopyright.Log(DisableDeprecationWarnings);
+
+	for (const FString& FWDecl : InOutReferenceGatherers.ForwardDeclarations)
+	{
+		if (FWDecl.Len() > 0)
+		{
+			GeneratedHeaderTextWithCopyright.Logf(TEXT("%s\r\n"), *FWDecl);
+		}
+	}
+
+	GeneratedHeaderTextWithCopyright.Log(InBodyText);
+	GeneratedHeaderTextWithCopyright.Log(EnableDeprecationWarnings);
+
+
+	const bool bHasChanged = SaveHeaderIfChanged(InOutReferenceGatherers, FileInfo, MoveTemp(GeneratedHeaderTextWithCopyright), OutSaveTempTask);
+	return bHasChanged;
+}
+
+/**
+ * Returns a string in the format CLASS_Something|CLASS_Something which represents all class flags that are set for the specified
+ * class which need to be exported as part of the DECLARE_CLASS macro
+ */
+FString FNativeClassHeaderGenerator::GetClassFlagExportText( UClass* Class )
+{
+	FString StaticClassFlagText;
+
+	check(Class);
+	if ( Class->HasAnyClassFlags(CLASS_Transient) )
+	{
+		StaticClassFlagText += TEXT(" | CLASS_Transient");
+	}
+	if( Class->HasAnyClassFlags(CLASS_DefaultConfig) )
+	{
+		StaticClassFlagText += TEXT(" | CLASS_DefaultConfig");
+	}
+	if( Class->HasAnyClassFlags(CLASS_GlobalUserConfig) )
+	{
+		StaticClassFlagText += TEXT(" | CLASS_GlobalUserConfig");
+	}
+	if( Class->HasAnyClassFlags(CLASS_Config) )
+	{
+		StaticClassFlagText += TEXT(" | CLASS_Config");
+	}
+	if ( Class->HasAnyClassFlags(CLASS_Interface) )
+	{
+		StaticClassFlagText += TEXT(" | CLASS_Interface");
+	}
+	if ( Class->HasAnyClassFlags(CLASS_Deprecated) )
+	{
+		StaticClassFlagText += TEXT(" | CLASS_Deprecated");
+	}
+
+	return StaticClassFlagText;
+}
+
+/**
+* Exports the header text for the list of enums specified
+*
+* @param	Enums	the enums to export
+*/
+void FNativeClassHeaderGenerator::ExportEnum(FOutputDevice& Out, UEnum* Enum) const
+{
+	// Export FOREACH macro
+	Out.Logf( TEXT("#define FOREACH_ENUM_%s(op) "), *Enum->GetName().ToUpper() );
+	bool bHasExistingMax = Enum->ContainsExistingMax();
+	int64 MaxEnumVal = bHasExistingMax ? Enum->GetMaxEnumValue() : 0;
+	for (int32 i = 0; i < Enum->NumEnums(); i++)
+	{
+		if (bHasExistingMax && Enum->GetValueByIndex(i) == MaxEnumVal)
+		{
+			continue;
+		}
+
+		const FString QualifiedEnumValue = Enum->GetNameByIndex(i).ToString();
+		Out.Logf( TEXT("\\\r\n\top(%s) "), *QualifiedEnumValue );
+	}
+	Out.Logf( TEXT("\r\n") );
+
+	// Forward declare the StaticEnum<> specialisation for enum classes
+	if (const EUnderlyingEnumType* EnumPropType = GEnumUnderlyingTypes.Find(Enum))
+	{
+		check(Enum->GetCppForm() == UEnum::ECppForm::EnumClass);
+
+		FString UnderlyingTypeString;
+
+		if (*EnumPropType != EUnderlyingEnumType::Unspecified)
+		{
+			UnderlyingTypeString = TEXT(" : ");
+
+			switch (*EnumPropType)
+			{
+			case EUnderlyingEnumType::int8:        UnderlyingTypeString += TNameOf<int8>::GetName();	break;
+			case EUnderlyingEnumType::int16:       UnderlyingTypeString += TNameOf<int16>::GetName();	break;
+			case EUnderlyingEnumType::int32:       UnderlyingTypeString += TNameOf<int32>::GetName();	break;
+			case EUnderlyingEnumType::int64:       UnderlyingTypeString += TNameOf<int64>::GetName();	break;
+			case EUnderlyingEnumType::uint8:       UnderlyingTypeString += TNameOf<uint8>::GetName();	break;
+			case EUnderlyingEnumType::uint16:      UnderlyingTypeString += TNameOf<uint16>::GetName();	break;
+			case EUnderlyingEnumType::uint32:      UnderlyingTypeString += TNameOf<uint32>::GetName();	break;
+			case EUnderlyingEnumType::uint64:      UnderlyingTypeString += TNameOf<uint64>::GetName();	break;
+			default:
+				check(false);
+			}
+		}
+
+		Out.Logf( TEXT("\r\n") );
+		Out.Logf( TEXT("enum class %s%s;\r\n"), *Enum->CppType, *UnderlyingTypeString );
+		Out.Logf( TEXT("template<> %sUEnum* StaticEnum<%s>();\r\n"), *GetAPIString(), *Enum->CppType );
+		Out.Logf( TEXT("\r\n") );
+	}
+}
+
+// Exports the header text for the list of structs specified (GENERATED_BODY impls)
+void FNativeClassHeaderGenerator::ExportGeneratedStructBodyMacros(FOutputDevice& OutGeneratedHeaderText, FOutputDevice& Out, FReferenceGatherers& OutReferenceGatherers, const FUnrealSourceFile& SourceFile, UScriptStruct* Struct) const
+{
+	const bool bIsDynamic = FClass::IsDynamic(Struct);
+	const FString ActualStructName = FNativeClassHeaderGenerator::GetOverriddenName(Struct);
+	const FString& FriendApiString  = GetAPIString();
+
+	UStruct* BaseStruct = Struct->GetSuperStruct();
+
+	const FString StructNameCPP = FNameLookupCPP::GetNameCPP(Struct);
+
+	const FString& SingletonName = GetSingletonName(Struct, OutReferenceGatherers.UniqueCrossModuleReferences);
+	const FString ChoppedSingletonName = SingletonName.LeftChop(2);
+
+	const FString RigVMParameterPrefix = TEXT("const FName& RigVMOperatorName, int32 RigVMOperatorIndex");
+	TArray<FString> RigVMVirtualFuncProlog, RigVMStubProlog;
+
+	// for RigVM methods we need to generated a macro used for implementing the static method
+	// and prepare two prologs: one for the virtual function implementation, and one for the stub
+	// invoking the static method.
+	const FRigVMStructInfo* StructRigVMInfo = FHeaderParser::StructRigVMMap.Find(Struct);
+	if(StructRigVMInfo)
+	{
+		//RigVMStubProlog.Add(FString::Printf(TEXT("ensure(RigVMOperandMemory.Num() == %d);"), StructRigVMInfo->Members.Num()));
+
+		int32 OperandIndex = 0;
+		for (int32 ParameterIndex = 0; ParameterIndex < StructRigVMInfo->Members.Num(); ParameterIndex++)
+		{
+			const FRigVMParameter& Parameter = StructRigVMInfo->Members[ParameterIndex];
+			if(Parameter.RequiresCast())
+			{
+				if (Parameter.IsArray() && !Parameter.IsConst() && !Parameter.MaxArraySize.IsEmpty())
+				{
+					RigVMVirtualFuncProlog.Add(FString::Printf(TEXT("%s.SetNum( %s );"), *Parameter.Name, *Parameter.MaxArraySize));
+				}
+				RigVMVirtualFuncProlog.Add(FString::Printf(TEXT("%s %s(%s);"), *Parameter.CastType, *Parameter.CastName, *Parameter.Name));
+			}
+
+			FString VariableType;
+			FString ExtractedType;
+			const FString& ParamTypeOriginal = Parameter.TypeOriginal(true);
+			const FString& ParamNameOriginal = Parameter.NameOriginal(false);
+
+			if (ParamTypeOriginal.StartsWith(TEXT("TArrayView"), ESearchCase::CaseSensitive))
+			{
+				ExtractedType = Parameter.ExtendedType().LeftChop(1).RightChop(1);
+				VariableType = ParamTypeOriginal;
+				
+				RigVMStubProlog.Add(FString::Printf(TEXT("%s %s = TArrayView<%s>((%s*)RigVMOperandMemory[%d], reinterpret_cast<uint64>(RigVMOperandMemory[%d]));"),
+					*VariableType,
+					*ParamNameOriginal,
+					*ExtractedType,
+					*ExtractedType,
+					OperandIndex,
+					OperandIndex + 1));
+
+				OperandIndex += 2;
+			}
+			else
+			{
+				VariableType = Parameter.TypeVariableRef(true);
+				ExtractedType = Parameter.TypeOriginal();
+				
+				FString ParameterCast = FString::Printf(TEXT("*(%s*)"), *ExtractedType);
+
+				RigVMStubProlog.Add(FString::Printf(TEXT("%s %s = %sRigVMOperandMemory[%d];"),
+				*VariableType,
+				*ParamNameOriginal,
+					*ParameterCast,
+					OperandIndex));
+
+				OperandIndex++;
+			}
+		}
+
+		FString StructMembers = StructRigVMInfo->Members.Declarations(false, TEXT(", \\\r\n\t\t"), true, false);
+
+		OutGeneratedHeaderText.Log(TEXT("\n"));
+		for (const FRigVMMethodInfo& MethodInfo : StructRigVMInfo->Methods)
+		{
+			FString ParameterSuffix = MethodInfo.Parameters.Declarations(true, TEXT(", \\\r\n\t\t"));
+			FString RigVMParameterPrefix2 = RigVMParameterPrefix + FString((StructMembers.IsEmpty() && ParameterSuffix.IsEmpty()) ? TEXT("") : TEXT(", \\\r\n\t\t"));
+			OutGeneratedHeaderText.Logf(TEXT("#define %s_%s() \\\r\n"), *StructNameCPP, *MethodInfo.Name);
+			OutGeneratedHeaderText.Logf(TEXT("\t%s %s::Static%s( \\\r\n\t\t%s%s%s \\\r\n\t)\n"), *MethodInfo.ReturnType, *StructNameCPP, *MethodInfo.Name, *RigVMParameterPrefix2, *StructMembers, *ParameterSuffix);
+		}
+		OutGeneratedHeaderText.Log(TEXT("\n"));
+	}
+
+	// Export struct.
+	if (Struct->StructFlags & STRUCT_Native)
+	{
+		check(Struct->StructMacroDeclaredLineNumber != INDEX_NONE);
+
+		const bool bRequiredAPI = !(Struct->StructFlags & STRUCT_RequiredAPI);
+
+		const FString FriendLine = FString::Printf(TEXT("\tfriend struct %s_Statics;\r\n"), *ChoppedSingletonName);
+		const FString StaticClassLine = FString::Printf(TEXT("\t%sstatic class UScriptStruct* StaticStruct();\r\n"), (bRequiredAPI ? *FriendApiString : TEXT("")));
+		const FString PrivatePropertiesOffset = PrivatePropertiesOffsetGetters(Struct, StructNameCPP);
+		
+		// if we have RigVM methods on this struct we need to 
+		// declare the static method as well as the stub method
+		FString RigVMMethodsDeclarations;
+		if (StructRigVMInfo)
+		{
+			FString StructMembers = StructRigVMInfo->Members.Declarations(false, TEXT(",\r\n\t\t"), true, false);
+			for (const FRigVMMethodInfo& MethodInfo : StructRigVMInfo->Methods)
+			{
+				FString StructMembersForStub = StructRigVMInfo->Members.Names(false, TEXT(",\r\n\t\t\t"), false);
+				FString ParameterSuffix = MethodInfo.Parameters.Declarations(true, TEXT(",\r\n\t\t"));
+				FString ParameterNamesSuffix = MethodInfo.Parameters.Names(true, TEXT(",\r\n\t\t\t"));
+				FString RigVMParameterPrefix2 = RigVMParameterPrefix + FString((StructMembers.IsEmpty() && ParameterSuffix.IsEmpty()) ? TEXT("") : TEXT(",\r\n\t\t"));
+				FString RigVMParameterPrefix4 = FString(TEXT("RigVMOperatorName,\r\n\t\t\tRigVMOperatorIndex")) + FString((StructMembersForStub.IsEmpty() && ParameterSuffix.IsEmpty()) ? TEXT("") : TEXT(",\r\n\t\t\t"));
+
+				RigVMMethodsDeclarations += FString::Printf(TEXT("\tstatic %s Static%s(\r\n\t\t%s%s%s\r\n\t);\r\n"), *MethodInfo.ReturnType, *MethodInfo.Name, *RigVMParameterPrefix2, *StructMembers, *ParameterSuffix);
+				RigVMMethodsDeclarations += FString::Printf(TEXT("\tFORCEINLINE static %s RigVM%s(\r\n\t\t%s,\r\n\t\tFRigVMOperandMemory RigVMOperandMemory,\r\n\t\tconst FRigVMUserDataArray& RigVMUserData\r\n\t)\r\n"), *MethodInfo.ReturnType, *MethodInfo.Name, *RigVMParameterPrefix);
+				RigVMMethodsDeclarations += FString::Printf(TEXT("\t{\r\n"));
+
+				// implement inline stub method body
+				if (MethodInfo.Parameters.Num() > 0)
+				{
+					//RigVMMethodsDeclarations += FString::Printf(TEXT("\t\tensure(RigVMUserData.Num() == %d);\r\n"), MethodInfo.Parameters.Num());
+					for (int32 ParameterIndex = 0; ParameterIndex < MethodInfo.Parameters.Num(); ParameterIndex++)
+					{
+						const FRigVMParameter& Parameter = MethodInfo.Parameters[ParameterIndex];
+						RigVMMethodsDeclarations += FString::Printf(TEXT("\t\t%s = *(%s*)RigVMUserData[%d];\r\n"), *Parameter.Declaration(), *Parameter.TypeNoRef(), ParameterIndex);
+					}
+					RigVMMethodsDeclarations += FString::Printf(TEXT("\t\t\r\n"));
+				}
+
+				if (RigVMStubProlog.Num() > 0)
+				{
+					for (const FString& RigVMStubPrologLine : RigVMStubProlog)
+					{
+						RigVMMethodsDeclarations += FString::Printf(TEXT("\t\t%s\r\n"), *RigVMStubPrologLine);
+					}
+					RigVMMethodsDeclarations += FString::Printf(TEXT("\t\t\r\n"));
+				}
+
+				RigVMMethodsDeclarations += FString::Printf(TEXT("\t\t%sStatic%s(\r\n\t\t\t%s%s%s\r\n\t\t);\r\n"), *MethodInfo.ReturnPrefix(), *MethodInfo.Name, *RigVMParameterPrefix4, *StructMembersForStub, *ParameterNamesSuffix);
+				RigVMMethodsDeclarations += FString::Printf(TEXT("\t}\r\n"));
+			}
+
+			for (const FRigVMParameter& StructMember : StructRigVMInfo->Members)
+			{
+				if (!StructMember.MaxArraySize.IsEmpty())
+				{
+					RigVMMethodsDeclarations += TEXT("\tvirtual int32 GetMaxArraySize(const FName& InMemberName, const FRigVMUserDataArray& RigVMUserData) override;\r\n");
+					break;
+				}
+			}
+		}
+
+		const FString SuperTypedef = BaseStruct ? FString::Printf(TEXT("\ttypedef %s Super;\r\n"), *FNameLookupCPP::GetNameCPP(BaseStruct)) : FString();
+
+		FString CombinedLine = FString::Printf(TEXT("%s%s%s%s%s"), *FriendLine, *StaticClassLine, *RigVMMethodsDeclarations, *PrivatePropertiesOffset, *SuperTypedef);
+		const FString MacroName = SourceFile.GetGeneratedBodyMacroName(Struct->StructMacroDeclaredLineNumber);
+
+		const FString Macroized = Macroize(*MacroName, MoveTemp(CombinedLine));
+		OutGeneratedHeaderText.Log(Macroized);
+
+		FString GetHashName = FString::Printf(TEXT("Get_%s_Hash"), *ChoppedSingletonName);
+
+		Out.Logf(TEXT("class UScriptStruct* %s::StaticStruct()\r\n"), *StructNameCPP);
+		Out.Logf(TEXT("{\r\n"));
+
+		// UStructs can have UClass or UPackage outer (if declared in non-UClass headers).
+		const FString& OuterName (bIsDynamic ? STRING_StructPackage : GetPackageSingletonName(CastChecked<UPackage>(Struct->GetOuter()), OutReferenceGatherers.UniqueCrossModuleReferences));
+		if (!bIsDynamic)
+		{
+			Out.Logf(TEXT("\tstatic class UScriptStruct* Singleton = NULL;\r\n"));
+		}
+		else
+		{
+			Out.Logf(TEXT("\tclass UPackage* %s = FindOrConstructDynamicTypePackage(TEXT(\"%s\"));\r\n"), *OuterName, *FClass::GetTypePackageName(Struct));
+			Out.Logf(TEXT("\tclass UScriptStruct* Singleton = Cast<UScriptStruct>(StaticFindObjectFast(UScriptStruct::StaticClass(), %s, TEXT(\"%s\")));\r\n"), *OuterName, *ActualStructName);
+		}
+
+		Out.Logf(TEXT("\tif (!Singleton)\r\n"));
+		Out.Logf(TEXT("\t{\r\n"));
+		Out.Logf(TEXT("\t\textern %suint32 %s();\r\n"), *FriendApiString, *GetHashName);
+
+		Out.Logf(TEXT("\t\tSingleton = GetStaticStruct(%s, %s, TEXT(\"%s\"), sizeof(%s), %s());\r\n"),
+			*ChoppedSingletonName, *OuterName, *ActualStructName, *StructNameCPP, *GetHashName);
+
+		// if this struct has RigVM methods - we need to register the method to our central
+		// registry on construction of the static struct
+		if (StructRigVMInfo)
+		{
+			for (const FRigVMMethodInfo& MethodInfo : StructRigVMInfo->Methods)
+			{
+				Out.Logf(TEXT("\t\tFRigVMRegistry::Get().Register(TEXT(\"%s::%s\"), &%s::RigVM%s, Singleton);\r\n"),
+					*StructNameCPP, *MethodInfo.Name, *StructNameCPP, *MethodInfo.Name);
+			}
+		}
+
+		Out.Logf(TEXT("\t}\r\n"));
+		Out.Logf(TEXT("\treturn Singleton;\r\n"));
+		Out.Logf(TEXT("}\r\n"));
+
+		// Forward declare the StaticStruct specialisation in the header
+		OutGeneratedHeaderText.Logf(TEXT("template<> %sUScriptStruct* StaticStruct<struct %s>();\r\n\r\n"), *GetAPIString(), *StructNameCPP);
+
+		// Generate the StaticStruct specialisation
+		Out.Logf(TEXT("template<> %sUScriptStruct* StaticStruct<%s>()\r\n"), *GetAPIString(), *StructNameCPP);
+		Out.Logf(TEXT("{\r\n"));
+		Out.Logf(TEXT("\treturn %s::StaticStruct();\r\n"), *StructNameCPP);
+		Out.Logf(TEXT("}\r\n"));
+
+		if (bIsDynamic)
+		{
+			const FString& StructPackageName = FClass::GetTypePackageName(Struct);
+			Out.Logf(TEXT("static FCompiledInDeferStruct Z_CompiledInDeferStruct_UScriptStruct_%s(%s::StaticStruct, TEXT(\"%s\"), TEXT(\"%s\"), true, %s, %s);\r\n"),
+				*StructNameCPP,
+				*StructNameCPP,
+				*StructPackageName,
+				*ActualStructName,
+				*AsTEXT(StructPackageName),
+				*AsTEXT(FNativeClassHeaderGenerator::GetOverriddenPathName(Struct)));
+		}
+		else
+		{
+			Out.Logf(TEXT("static FCompiledInDeferStruct Z_CompiledInDeferStruct_UScriptStruct_%s(%s::StaticStruct, TEXT(\"%s\"), TEXT(\"%s\"), false, nullptr, nullptr);\r\n"),
+				*StructNameCPP,
+				*StructNameCPP,
+				*Struct->GetOutermost()->GetName(),
+				*ActualStructName);
+		}
+
+		// Generate StaticRegisterNatives equivalent for structs without classes.
+		if (!Struct->GetOuter()->IsA(UStruct::StaticClass()))
+		{
+			const FString ShortPackageName = FPackageName::GetShortName(Struct->GetOuter()->GetName());
+			Out.Logf(TEXT("static struct FScriptStruct_%s_StaticRegisterNatives%s\r\n"), *ShortPackageName, *StructNameCPP);
+			Out.Logf(TEXT("{\r\n"));
+			Out.Logf(TEXT("\tFScriptStruct_%s_StaticRegisterNatives%s()\r\n"), *ShortPackageName, *StructNameCPP);
+			Out.Logf(TEXT("\t{\r\n"));
+
+			Out.Logf(TEXT("\t\tUScriptStruct::DeferCppStructOps(FName(TEXT(\"%s\")),new UScriptStruct::TCppStructOps<%s>);\r\n"), *ActualStructName, *StructNameCPP);
+
+			Out.Logf(TEXT("\t}\r\n"));
+			Out.Logf(TEXT("} ScriptStruct_%s_StaticRegisterNatives%s;\r\n"), *ShortPackageName, *StructNameCPP);
+		}
+	}
+
+	FString StaticsStructName = ChoppedSingletonName + TEXT("_Statics");
+
+	FUHTStringBuilder GeneratedStructRegisterFunctionText;
+	FUHTStringBuilder StaticDefinitions;
+
+	GeneratedStructRegisterFunctionText.Logf(TEXT("\tstruct %s\r\n"), *StaticsStructName);
+	GeneratedStructRegisterFunctionText.Logf(TEXT("\t{\r\n"));
+
+	// if this is a no export struct, we will put a local struct here for offset determination
+	TArray<UScriptStruct*> NoExportStructs = FindNoExportStructs(Struct);
+	for (UScriptStruct* NoExportStruct : NoExportStructs)
+	{
+		ExportMirrorsForNoexportStruct(GeneratedStructRegisterFunctionText, NoExportStruct, /*Indent=*/ 2);
+	}
+
+	if (BaseStruct)
+	{
+		CastChecked<UScriptStruct>(BaseStruct); // this better actually be a script struct
+		GetSingletonName(BaseStruct, OutReferenceGatherers.UniqueCrossModuleReferences); // Call to potentially collect references
+	}
+
+	EStructFlags UncomputedFlags = (EStructFlags)(Struct->StructFlags & ~STRUCT_ComputedFlags);
+
+	FString OuterFunc;
+	if (!bIsDynamic)
+	{
+		OuterFunc = GetPackageSingletonName(CastChecked<UPackage>(Struct->GetOuter()), OutReferenceGatherers.UniqueCrossModuleReferences).LeftChop(2);
+	}
+	else
+	{
+		GeneratedStructRegisterFunctionText.Log(TEXT("\t\tstatic UObject* OuterFuncGetter();\r\n"));
+
+		StaticDefinitions.Logf(TEXT("\tUObject* %s::OuterFuncGetter()\r\n"), *StaticsStructName);
+		StaticDefinitions.Log (TEXT("\t{\r\n"));
+		StaticDefinitions.Logf(TEXT("\t\treturn FindOrConstructDynamicTypePackage(TEXT(\"%s\"));"), *FClass::GetTypePackageName(Struct));
+		StaticDefinitions.Log (TEXT("\t}\r\n"));
+
+		OuterFunc = TEXT("&OuterFuncGetter");
+	}
+
+	FString MetaDataParams = OutputMetaDataCodeForObject(GeneratedStructRegisterFunctionText, StaticDefinitions, Struct, *FString::Printf(TEXT("%s::Struct_MetaDataParams"), *StaticsStructName), TEXT("\t\t"), TEXT("\t"));
+
+	TArray<FProperty*> Props;
+	Algo::Copy(TFieldRange<FProperty>(Struct, EFieldIteratorFlags::ExcludeSuper), Props, Algo::NoRef);
+
+	FString NewStructOps;
+	if (Struct->StructFlags & STRUCT_Native)
+	{
+		GeneratedStructRegisterFunctionText.Log(TEXT("\t\tstatic void* NewStructOps();\r\n"));
+
+		StaticDefinitions.Logf(TEXT("\tvoid* %s::NewStructOps()\r\n"), *StaticsStructName);
+		StaticDefinitions.Log (TEXT("\t{\r\n"));
+		StaticDefinitions.Logf(TEXT("\t\treturn (UScriptStruct::ICppStructOps*)new UScriptStruct::TCppStructOps<%s>();\r\n"), *StructNameCPP);
+		StaticDefinitions.Log (TEXT("\t}\r\n"));
+
+		NewStructOps = TEXT("&NewStructOps");
+	}
+	else
+	{
+		NewStructOps = TEXT("nullptr");
+	}
+
+	TTuple<FString, FString> PropertyRange = OutputProperties(GeneratedStructRegisterFunctionText, StaticDefinitions, OutReferenceGatherers, *FString::Printf(TEXT("%s::"), *StaticsStructName), Props, TEXT("\t\t"), TEXT("\t"));
+
+	GeneratedStructRegisterFunctionText.Log (TEXT("\t\tstatic const UE4CodeGen_Private::FStructParams ReturnStructParams;\r\n"));
+
+	StaticDefinitions.Logf(TEXT("\tconst UE4CodeGen_Private::FStructParams %s::ReturnStructParams = {\r\n"), *StaticsStructName);
+	StaticDefinitions.Logf(TEXT("\t\t(UObject* (*)())%s,\r\n"), *OuterFunc);
+	StaticDefinitions.Logf(TEXT("\t\t%s,\r\n"), *GetSingletonNameFuncAddr(BaseStruct, OutReferenceGatherers.UniqueCrossModuleReferences));
+	StaticDefinitions.Logf(TEXT("\t\t%s,\r\n"), *NewStructOps);
+	StaticDefinitions.Logf(TEXT("\t\t%s,\r\n"), *CreateUTF8LiteralString(ActualStructName));
+	StaticDefinitions.Logf(TEXT("\t\tsizeof(%s),\r\n"), *StructNameCPP);
+	StaticDefinitions.Logf(TEXT("\t\talignof(%s),\r\n"), *StructNameCPP);
+	StaticDefinitions.Logf(TEXT("\t\t%s,\r\n"), *PropertyRange.Get<0>());
+	StaticDefinitions.Logf(TEXT("\t\t%s,\r\n"), *PropertyRange.Get<1>());
+	StaticDefinitions.Logf(TEXT("\t\t%s,\r\n"), bIsDynamic ? TEXT("RF_Public|RF_Transient") : TEXT("RF_Public|RF_Transient|RF_MarkAsNative"));
+	StaticDefinitions.Logf(TEXT("\t\tEStructFlags(0x%08X),\r\n"), (uint32)UncomputedFlags);
+	StaticDefinitions.Logf(TEXT("\t\t%s\r\n"), *MetaDataParams);
+	StaticDefinitions.Log (TEXT("\t};\r\n"));
+
+	GeneratedStructRegisterFunctionText.Log (TEXT("\t};\r\n"));
+
+	GeneratedStructRegisterFunctionText.Log(StaticDefinitions);
+
+	GeneratedStructRegisterFunctionText.Logf(TEXT("\tUScriptStruct* %s\r\n"), *SingletonName);
+	GeneratedStructRegisterFunctionText.Log (TEXT("\t{\r\n"));
+
+	FString NoExportStructNameCPP;
+	if (NoExportStructs.Contains(Struct))
+	{
+		NoExportStructNameCPP = FString::Printf(TEXT("%s::%s"), *StaticsStructName, *StructNameCPP);
+	}
+	else
+	{
+		NoExportStructNameCPP = StructNameCPP;
+	}
+
+	FString HashFuncName = FString::Printf(TEXT("Get_%s_Hash"), *SingletonName.Replace(TEXT("()"), TEXT(""), ESearchCase::CaseSensitive));
+	// Structs can either have a UClass or UPackage as outer (if declared in non-UClass header).
+	if (!bIsDynamic)
+	{
+		GeneratedStructRegisterFunctionText.Log (TEXT("#if WITH_HOT_RELOAD\r\n"));
+		GeneratedStructRegisterFunctionText.Logf(TEXT("\t\textern uint32 %s();\r\n"), *HashFuncName);
+		GeneratedStructRegisterFunctionText.Logf(TEXT("\t\tUPackage* Outer = %s;\r\n"), *GetPackageSingletonName(CastChecked<UPackage>(Struct->GetOuter()), OutReferenceGatherers.UniqueCrossModuleReferences));
+		GeneratedStructRegisterFunctionText.Logf(TEXT("\t\tstatic UScriptStruct* ReturnStruct = FindExistingStructIfHotReloadOrDynamic(Outer, TEXT(\"%s\"), sizeof(%s), %s(), false);\r\n"), *ActualStructName, *NoExportStructNameCPP, *HashFuncName);
+		GeneratedStructRegisterFunctionText.Log (TEXT("#else\r\n"));
+		GeneratedStructRegisterFunctionText.Logf(TEXT("\t\tstatic UScriptStruct* ReturnStruct = nullptr;\r\n"));
+		GeneratedStructRegisterFunctionText.Log (TEXT("#endif\r\n"));
+	}
+	else
+	{
+		GeneratedStructRegisterFunctionText.Logf(TEXT("\t\textern uint32 %s();\r\n"), *HashFuncName);
+		GeneratedStructRegisterFunctionText.Logf(TEXT("\t\tUPackage* Outer = FindOrConstructDynamicTypePackage(TEXT(\"%s\"));\r\n"), *FClass::GetTypePackageName(Struct));
+		GeneratedStructRegisterFunctionText.Logf(TEXT("\t\tUScriptStruct* ReturnStruct = FindExistingStructIfHotReloadOrDynamic(Outer, TEXT(\"%s\"), sizeof(%s), %s(), true);\r\n"), *ActualStructName, *NoExportStructNameCPP, *HashFuncName);
+	}
+	GeneratedStructRegisterFunctionText.Logf(TEXT("\t\tif (!ReturnStruct)\r\n"));
+	GeneratedStructRegisterFunctionText.Log (TEXT("\t\t{\r\n"));
+
+	GeneratedStructRegisterFunctionText.Logf(TEXT("\t\t\tUE4CodeGen_Private::ConstructUScriptStruct(ReturnStruct, %s::ReturnStructParams);\r\n"), *StaticsStructName);
+	GeneratedStructRegisterFunctionText.Log (TEXT("\t\t}\r\n"));
+	GeneratedStructRegisterFunctionText.Logf(TEXT("\t\treturn ReturnStruct;\r\n"));
+	GeneratedStructRegisterFunctionText.Log (TEXT("\t}\r\n"));
+
+	uint32 StructHash = GenerateTextHash(*GeneratedStructRegisterFunctionText);
+	AddGeneratedCodeHash(Struct, StructHash);
+
+	Out.Log(GeneratedStructRegisterFunctionText);
+	Out.Logf(TEXT("\tuint32 %s() { return %uU; }\r\n"), *HashFuncName, StructHash);
+
+	// if this struct has RigVM methods we need to implement both the 
+	// virtual function as well as the stub method here.
+	// The static method is implemented by the user using a macro.
+	if (StructRigVMInfo)
+	{
+		FString StructMembersForVirtualFunc = StructRigVMInfo->Members.Names(false, TEXT(",\r\n\t\t"), true);
+
+		for (const FRigVMMethodInfo& MethodInfo : StructRigVMInfo->Methods)
+		{
+			Out.Log(TEXT("\r\n"));
+
+			FString ParameterDeclaration = MethodInfo.Parameters.Declarations(false, TEXT(",\r\n\t\t"));
+			FString ParameterSuffix = MethodInfo.Parameters.Names(true, TEXT(",\r\n\t\t"));
+			FString RigVMParameterPrefix2 = RigVMParameterPrefix + FString((StructMembersForVirtualFunc.IsEmpty() && ParameterSuffix.IsEmpty()) ? TEXT("") : TEXT(",\r\n\t\t"));
+			FString RigVMParameterPrefix3 = FString(TEXT("NAME_None,\r\n\t\tINDEX_NONE")) + FString((StructMembersForVirtualFunc.IsEmpty() && ParameterSuffix.IsEmpty()) ? TEXT("") : TEXT(",\r\n\t\t"));
+
+			// implement the virtual function body.
+			Out.Logf(TEXT("%s %s::%s(%s)\r\n"), *MethodInfo.ReturnType, *StructNameCPP, *MethodInfo.Name, *ParameterDeclaration);
+			Out.Log(TEXT("{\r\n"));
+
+			if(RigVMVirtualFuncProlog.Num() > 0)
+			{
+				for (const FString& RigVMVirtualFuncPrologLine : RigVMVirtualFuncProlog)
+				{
+					Out.Logf(TEXT("\t%s\r\n"), *RigVMVirtualFuncPrologLine);
+				}
+				Out.Log(TEXT("\t\r\n"));
+			}
+
+			Out.Logf(TEXT("    %sStatic%s(\r\n\t\t%s%s%s\r\n\t);\n"), *MethodInfo.ReturnPrefix(), *MethodInfo.Name, *RigVMParameterPrefix3, *StructMembersForVirtualFunc, *ParameterSuffix);
+			Out.Log(TEXT("}\r\n"));
+		}
+
+			Out.Log(TEXT("\r\n"));
+
+		bool bHasGetMaxArraySize = false;
+		for (const FRigVMParameter& StructMember : StructRigVMInfo->Members)
+			{
+			if (!StructMember.MaxArraySize.IsEmpty())
+				{
+				bHasGetMaxArraySize = true;
+				break;
+				}
+			}
+
+		if (bHasGetMaxArraySize)
+			{
+			Out.Logf(TEXT("int32 %s::GetMaxArraySize(const FName& InMemberName, const FRigVMUserDataArray& RigVMUserData)\r\n"), *StructNameCPP);
+			Out.Log(TEXT("{\r\n"));
+			for (const FRigVMParameter& StructMember : StructRigVMInfo->Members)
+			{
+				if (!StructMember.MaxArraySize.IsEmpty())
+				{
+					Out.Logf(TEXT("\tif(InMemberName == TEXT(\"%s\"))\r\n"), *StructMember.Name);
+					Out.Log(TEXT("\t{\r\n"));
+					Out.Logf(TEXT("\t\treturn %s;\r\n"), *StructMember.MaxArraySize);
+					Out.Log(TEXT("\t}\r\n"));
+				}
+			}
+			Out.Log(TEXT("\treturn INDEX_NONE;\r\n"));
+			Out.Log(TEXT("}\r\n\r\n"));
+		}
+	}
+}
+
+void FNativeClassHeaderGenerator::ExportGeneratedEnumInitCode(FOutputDevice& Out, FReferenceGatherers& OutReferenceGatherers, const FUnrealSourceFile& SourceFile, UEnum* Enum) const
+{
+	const bool    bIsDynamic            = FClass::IsDynamic(static_cast<UField*>(Enum));
+	const FString SingletonName         = GetSingletonNameFuncAddr(Enum, OutReferenceGatherers.UniqueCrossModuleReferences);
+	const FString EnumNameCpp           = Enum->GetName(); //UserDefinedEnum should already have a valid cpp name.
+	const FString OverriddenEnumNameCpp = FNativeClassHeaderGenerator::GetOverriddenName(Enum);
+
+	const bool bIsEditorOnlyDataType = GEditorOnlyDataTypes.Contains(Enum);
+
+	FMacroBlockEmitter EditorOnlyData(Out, TEXT("WITH_EDITORONLY_DATA"));
+	EditorOnlyData(bIsEditorOnlyDataType);
+
+	const FString& PackageSingletonName = (bIsDynamic ? FClass::GetTypePackageName(static_cast<UField*>(Enum)) : GetPackageSingletonName(CastChecked<UPackage>(Enum->GetOuter()), OutReferenceGatherers.UniqueCrossModuleReferences));
+
+	Out.Logf(TEXT("\tstatic UEnum* %s_StaticEnum()\r\n"), *Enum->GetName());
+	Out.Logf(TEXT("\t{\r\n"));
+
+	if (!bIsDynamic)
+	{
+		Out.Logf(TEXT("\t\tstatic UEnum* Singleton = nullptr;\r\n"));
+	}
+	else
+	{
+		Out.Logf(TEXT("\t\tclass UPackage* EnumPackage = FindOrConstructDynamicTypePackage(TEXT(\"%s\"));\r\n"), *PackageSingletonName);
+		Out.Logf(TEXT("\t\tclass UEnum* Singleton = Cast<UEnum>(StaticFindObjectFast(UEnum::StaticClass(), EnumPackage, TEXT(\"%s\")));\r\n"), *OverriddenEnumNameCpp);
+	}
+	Out.Logf(TEXT("\t\tif (!Singleton)\r\n"));
+	Out.Logf(TEXT("\t\t{\r\n"));
+	if (!bIsDynamic)
+	{
+		Out.Logf(TEXT("\t\t\tSingleton = GetStaticEnum(%s, %s, TEXT(\"%s\"));\r\n"), *SingletonName, *PackageSingletonName, *Enum->GetName());
+	}
+	else
+	{
+		Out.Logf(TEXT("\t\t\tSingleton = GetStaticEnum(%s, EnumPackage, TEXT(\"%s\"));\r\n"), *SingletonName, *OverriddenEnumNameCpp);
+	}
+
+	Out.Logf(TEXT("\t\t}\r\n"));
+	Out.Logf(TEXT("\t\treturn Singleton;\r\n"));
+	Out.Logf(TEXT("\t}\r\n"));
+
+	Out.Logf(TEXT("\ttemplate<> %sUEnum* StaticEnum<%s>()\r\n"), *GetAPIString(), *Enum->CppType);
+	Out.Logf(TEXT("\t{\r\n"));
+	Out.Logf(TEXT("\t\treturn %s_StaticEnum();\r\n"), *Enum->GetName());
+	Out.Logf(TEXT("\t}\r\n"));
+
+	if (bIsDynamic)
+	{
+		const FString& EnumPackageName = FClass::GetTypePackageName(static_cast<UField*>(Enum));
+		Out.Logf(
+			TEXT("\tstatic FCompiledInDeferEnum Z_CompiledInDeferEnum_UEnum_%s(%s_StaticEnum, TEXT(\"%s\"), TEXT(\"%s\"), true, %s, %s);\r\n"),
+			*EnumNameCpp,
+			*EnumNameCpp,
+			*EnumPackageName,
+			*OverriddenEnumNameCpp,
+			*AsTEXT(EnumPackageName),
+			*AsTEXT(FNativeClassHeaderGenerator::GetOverriddenPathName(static_cast<UField*>(Enum)))
+		);
+	}
+	else
+	{
+		Out.Logf(
+			TEXT("\tstatic FCompiledInDeferEnum Z_CompiledInDeferEnum_UEnum_%s(%s_StaticEnum, TEXT(\"%s\"), TEXT(\"%s\"), false, nullptr, nullptr);\r\n"),
+			*EnumNameCpp,
+			*EnumNameCpp,
+			*Enum->GetOutermost()->GetName(),
+			*OverriddenEnumNameCpp
+		);
+	}
+
+	const FString& EnumSingletonName = GetSingletonName(Enum, OutReferenceGatherers.UniqueCrossModuleReferences);
+	const FString HashFuncName       = FString::Printf(TEXT("Get_%s_Hash"), *SingletonName);
+
+	FUHTStringBuilder GeneratedEnumRegisterFunctionText;
+
+	GeneratedEnumRegisterFunctionText.Logf(TEXT("\tUEnum* %s\r\n"), *EnumSingletonName);
+	GeneratedEnumRegisterFunctionText.Logf(TEXT("\t{\r\n"));
+
+	// Enums can either have a UClass or UPackage as outer (if declared in non-UClass header).
+	FString OuterString;
+	if (!bIsDynamic)
+	{
+		OuterString = PackageSingletonName;
+		GeneratedEnumRegisterFunctionText.Logf(TEXT("#if WITH_HOT_RELOAD\r\n"));
+		GeneratedEnumRegisterFunctionText.Logf(TEXT("\t\tUPackage* Outer = %s;\r\n"), *OuterString);
+		GeneratedEnumRegisterFunctionText.Logf(TEXT("\t\tstatic UEnum* ReturnEnum = FindExistingEnumIfHotReloadOrDynamic(Outer, TEXT(\"%s\"), 0, %s(), false);\r\n"), *EnumNameCpp, *HashFuncName);
+		GeneratedEnumRegisterFunctionText.Logf(TEXT("#else\r\n"));
+		GeneratedEnumRegisterFunctionText.Logf(TEXT("\t\tstatic UEnum* ReturnEnum = nullptr;\r\n"));
+		GeneratedEnumRegisterFunctionText.Logf(TEXT("#endif // WITH_HOT_RELOAD\r\n"));
+	}
+	else
+	{
+		OuterString = FString::Printf(TEXT("[](){ return (UObject*)FindOrConstructDynamicTypePackage(TEXT(\"%s\")); }()"), *PackageSingletonName);
+		GeneratedEnumRegisterFunctionText.Logf(TEXT("\t\tUPackage* Outer = FindOrConstructDynamicTypePackage(TEXT(\"%s\"));"), *PackageSingletonName);
+		GeneratedEnumRegisterFunctionText.Logf(TEXT("\t\tUEnum* ReturnEnum = FindExistingEnumIfHotReloadOrDynamic(Outer, TEXT(\"%s\"), 0, %s(), true);\r\n"), *OverriddenEnumNameCpp, *HashFuncName);
+	}
+	GeneratedEnumRegisterFunctionText.Logf(TEXT("\t\tif (!ReturnEnum)\r\n"));
+	GeneratedEnumRegisterFunctionText.Logf(TEXT("\t\t{\r\n"));
+
+	const TCHAR* UEnumObjectFlags = bIsDynamic ? TEXT("RF_Public|RF_Transient") : TEXT("RF_Public|RF_Transient|RF_MarkAsNative");
+
+	const TCHAR* EnumFormStr = TEXT("");
+	switch (Enum->GetCppForm())
+	{
+		case UEnum::ECppForm::Regular:    EnumFormStr = TEXT("UEnum::ECppForm::Regular");    break;
+		case UEnum::ECppForm::Namespaced: EnumFormStr = TEXT("UEnum::ECppForm::Namespaced"); break;
+		case UEnum::ECppForm::EnumClass:  EnumFormStr = TEXT("UEnum::ECppForm::EnumClass");  break;
+	}
+
+	const FString& EnumDisplayNameFn = Enum->GetMetaData(TEXT("EnumDisplayNameFn"));
+
+	GeneratedEnumRegisterFunctionText.Logf(TEXT("\t\t\tstatic const UE4CodeGen_Private::FEnumeratorParam Enumerators[] = {\r\n"));
+	for (int32 Index = 0; Index != Enum->NumEnums(); ++Index)
+	{
+		const TCHAR* OverridenNameMetaDatakey = TEXT("OverrideName");
+		const FString KeyName = Enum->HasMetaData(OverridenNameMetaDatakey, Index) ? Enum->GetMetaData(OverridenNameMetaDatakey, Index) : Enum->GetNameByIndex(Index).ToString();
+		GeneratedEnumRegisterFunctionText.Logf(TEXT("\t\t\t\t{ %s, (int64)%s },\r\n"), *CreateUTF8LiteralString(KeyName), *Enum->GetNameByIndex(Index).ToString());
+	}
+	GeneratedEnumRegisterFunctionText.Logf(TEXT("\t\t\t};\r\n"));
+
+	FOutputDeviceNull OutputDeviceNull;
+	FString MetaDataParams = OutputMetaDataCodeForObject(OutputDeviceNull, GeneratedEnumRegisterFunctionText, Enum, TEXT("Enum_MetaDataParams"), TEXT(""), TEXT("\t\t\t"));
+
+	GeneratedEnumRegisterFunctionText.Logf(TEXT("\t\t\tstatic const UE4CodeGen_Private::FEnumParams EnumParams = {\r\n"));
+	GeneratedEnumRegisterFunctionText.Logf(TEXT("\t\t\t\t(UObject*(*)())%s,\r\n"), *OuterString.LeftChop(2));
+	GeneratedEnumRegisterFunctionText.Logf(TEXT("\t\t\t\t%s,\r\n"), EnumDisplayNameFn.IsEmpty() ? TEXT("nullptr") : *EnumDisplayNameFn);
+	GeneratedEnumRegisterFunctionText.Logf(TEXT("\t\t\t\t%s,\r\n"), *CreateUTF8LiteralString(OverriddenEnumNameCpp));
+	GeneratedEnumRegisterFunctionText.Logf(TEXT("\t\t\t\t%s,\r\n"), *CreateUTF8LiteralString(Enum->CppType));
+	GeneratedEnumRegisterFunctionText.Logf(TEXT("\t\t\t\tEnumerators,\r\n"));
+	GeneratedEnumRegisterFunctionText.Logf(TEXT("\t\t\t\tUE_ARRAY_COUNT(Enumerators),\r\n"));
+	GeneratedEnumRegisterFunctionText.Logf(TEXT("\t\t\t\t%s,\r\n"), UEnumObjectFlags);
+	GeneratedEnumRegisterFunctionText.Logf(TEXT("\t\t\t\tUE4CodeGen_Private::EDynamicType::%s,\r\n"), bIsDynamic ? TEXT("Dynamic") : TEXT("NotDynamic"));
+	GeneratedEnumRegisterFunctionText.Logf(TEXT("\t\t\t\t(uint8)%s,\r\n"), EnumFormStr);
+	GeneratedEnumRegisterFunctionText.Logf(TEXT("\t\t\t\t%s\r\n"), *MetaDataParams);
+	GeneratedEnumRegisterFunctionText.Logf(TEXT("\t\t\t};\r\n"));
+	GeneratedEnumRegisterFunctionText.Logf(TEXT("\t\t\tUE4CodeGen_Private::ConstructUEnum(ReturnEnum, EnumParams);\r\n"));
+	GeneratedEnumRegisterFunctionText.Logf(TEXT("\t\t}\r\n"));
+	GeneratedEnumRegisterFunctionText.Logf(TEXT("\t\treturn ReturnEnum;\r\n"));
+	GeneratedEnumRegisterFunctionText.Logf(TEXT("\t}\r\n"));
+
+	uint32 EnumHash = GenerateTextHash(*GeneratedEnumRegisterFunctionText);
+	Out.Logf(TEXT("\tuint32 %s() { return %uU; }\r\n"), *HashFuncName, EnumHash);
+	Out.Log(GeneratedEnumRegisterFunctionText);
+}
+
+void FNativeClassHeaderGenerator::ExportMirrorsForNoexportStruct(FOutputDevice& Out, UScriptStruct* Struct, int32 TextIndent)
+{
+	// Export struct.
+	const FString StructName = FNameLookupCPP::GetNameCPP(Struct);
+	Out.Logf(TEXT("%sstruct %s"), FCString::Tab(TextIndent), *StructName);
+	if (Struct->GetSuperStruct() != NULL)
+	{
+		Out.Logf(TEXT(" : public %s"), *FNameLookupCPP::GetNameCPP(Struct->GetSuperStruct()));
+	}
+	Out.Logf(TEXT("\r\n%s{\r\n"), FCString::Tab(TextIndent));
+
+	// Export the struct's CPP properties.
+	ExportProperties(Out, Struct, TextIndent);
+
+	Out.Logf(TEXT("%s};\r\n\r\n"), FCString::Tab(TextIndent));
+}
+
+bool FNativeClassHeaderGenerator::WillExportEventParms( UFunction* Function )
+{
+  TFieldIterator<FProperty> It(Function);
+  return It && (It->PropertyFlags&CPF_Parm);
+}
+
+void WriteEventFunctionPrologue(FOutputDevice& Output, int32 Indent, const FParmsAndReturnProperties& Parameters, UObject* FunctionOuter, const TCHAR* FunctionName)
+{
+	// now the body - first we need to declare a struct which will hold the parameters for the event/delegate call
+	Output.Logf(TEXT("\r\n%s{\r\n"), FCString::Tab(Indent));
+
+	// declare and zero-initialize the parameters and return value, if applicable
+	if (!Parameters.HasParms())
+		return;
+
+	FString EventStructName = GetEventStructParamsName(FunctionOuter, FunctionName);
+
+	Output.Logf(TEXT("%s%s Parms;\r\n"), FCString::Tab(Indent + 1), *EventStructName );
+
+	// Declare a parameter struct for this event/delegate and assign the struct members using the values passed into the event/delegate call.
+	for (FProperty* Prop : Parameters.Parms)
+	{
+		const FString PropertyName = Prop->GetName();
+		if (Prop->ArrayDim > 1)
+		{
+			Output.Logf(TEXT("%sFMemory::Memcpy(Parms.%s,%s,sizeof(Parms.%s));\r\n"), FCString::Tab(Indent + 1), *PropertyName, *PropertyName, *PropertyName);
+		}
+		else
+		{
+			FString ValueAssignmentText = PropertyName;
+			if (Prop->IsA<FBoolProperty>())
+			{
+				ValueAssignmentText += TEXT(" ? true : false");
+			}
+
+			Output.Logf(TEXT("%sParms.%s=%s;\r\n"), FCString::Tab(Indent + 1), *PropertyName, *ValueAssignmentText);
+		}
+	}
+}
+
+void WriteEventFunctionEpilogue(FOutputDevice& Output, int32 Indent, const FParmsAndReturnProperties& Parameters)
+{
+	// Out parm copying.
+	for (FProperty* Prop : Parameters.Parms)
+	{
+		if ((Prop->PropertyFlags & (CPF_OutParm | CPF_ConstParm)) == CPF_OutParm)
+		{
+			const FString PropertyName = Prop->GetName();
+			if ( Prop->ArrayDim > 1 )
+			{
+				Output.Logf(TEXT("%sFMemory::Memcpy(&%s,&Parms.%s,sizeof(%s));\r\n"), FCString::Tab(Indent + 1), *PropertyName, *PropertyName, *PropertyName);
+			}
+			else
+			{
+				Output.Logf(TEXT("%s%s=Parms.%s;\r\n"), FCString::Tab(Indent + 1), *PropertyName, *PropertyName);
+			}
+		}
+	}
+
+	// Return value.
+	if (Parameters.Return)
+	{
+		// Make sure uint32 -> bool is supported
+		bool bBoolProperty = Parameters.Return->IsA(FBoolProperty::StaticClass());
+		Output.Logf(TEXT("%sreturn %sParms.%s;\r\n"), FCString::Tab(Indent + 1), bBoolProperty ? TEXT("!!") : TEXT(""), *Parameters.Return->GetName());
+	}
+	Output.Logf(TEXT("%s}\r\n"), FCString::Tab(Indent));
+}
+
+void FNativeClassHeaderGenerator::ExportDelegateDeclaration(FOutputDevice& Out, FReferenceGatherers& OutReferenceGatherers, const FUnrealSourceFile& SourceFile, UFunction* Function) const
+{
+	static const TCHAR DelegateStr[] = TEXT("delegate");
+
+	check(Function->HasAnyFunctionFlags(FUNC_Delegate));
+
+	const bool bIsMulticastDelegate = Function->HasAnyFunctionFlags( FUNC_MulticastDelegate );
+
+	// Unmangle the function name
+	const FString DelegateName = Function->GetName().LeftChop(HEADER_GENERATED_DELEGATE_SIGNATURE_SUFFIX_LENGTH);
+
+	const FFunctionData* CompilerInfo = FFunctionData::FindForFunction(Function);
+
+	FFuncInfo FunctionData = CompilerInfo->GetFunctionData();
+
+	// Add class name to beginning of function, to avoid collisions with other classes with the same delegate name in this scope
+	check(FunctionData.MarshallAndCallName.StartsWith(DelegateStr));
+	FString ShortName = *FunctionData.MarshallAndCallName + UE_ARRAY_COUNT(DelegateStr) - 1;
+	FunctionData.MarshallAndCallName = FString::Printf( TEXT( "F%s_DelegateWrapper" ), *ShortName );
+
+	// Setup delegate parameter
+	const FString ExtraParam = FString::Printf(
+		TEXT( "const %s& %s" ),
+		bIsMulticastDelegate ? TEXT( "FMulticastScriptDelegate" ) : TEXT( "FScriptDelegate" ),
+		*DelegateName
+	);
+
+	FUHTStringBuilder DelegateOutput;
+	DelegateOutput.Log(TEXT("static "));
+
+	// export the line that looks like: int32 Main(const FString& Parms)
+	ExportNativeFunctionHeader(DelegateOutput, OutReferenceGatherers.ForwardDeclarations, FunctionData, EExportFunctionType::Event, EExportFunctionHeaderStyle::Declaration, *ExtraParam, *GetAPIString());
+
+	// Only exporting function prototype
+	DelegateOutput.Logf(TEXT(";\r\n"));
+
+	ExportFunction(Out, OutReferenceGatherers, SourceFile, Function, false);
+}
+
+void FNativeClassHeaderGenerator::ExportDelegateDefinition(FOutputDevice& Out, FReferenceGatherers& OutReferenceGatherers, const FUnrealSourceFile& SourceFile, UFunction* Function) const
+{
+	const TCHAR DelegateStr[] = TEXT("delegate");
+
+	check(Function->HasAnyFunctionFlags(FUNC_Delegate));
+
+	// Export parameters structs for all delegates.  We'll need these to declare our delegate execution function.
+	FUHTStringBuilder DelegateOutput;
+	ExportEventParm(DelegateOutput, OutReferenceGatherers.ForwardDeclarations, Function, /*Indent=*/ 0, /*bOutputConstructor=*/ true, EExportingState::Normal);
+
+	const bool bIsMulticastDelegate = Function->HasAnyFunctionFlags( FUNC_MulticastDelegate );
+
+	// Unmangle the function name
+	const FString DelegateName = Function->GetName().LeftChop(HEADER_GENERATED_DELEGATE_SIGNATURE_SUFFIX_LENGTH);
+
+	const FFunctionData* CompilerInfo = FFunctionData::FindForFunction(Function);
+
+	FFuncInfo FunctionData = CompilerInfo->GetFunctionData();
+
+	// Always export delegate wrapper functions as inline
+	FunctionData.FunctionExportFlags |= FUNCEXPORT_Inline;
+
+	// Add class name to beginning of function, to avoid collisions with other classes with the same delegate name in this scope
+	check(FunctionData.MarshallAndCallName.StartsWith(DelegateStr));
+	FString ShortName = *FunctionData.MarshallAndCallName + UE_ARRAY_COUNT(DelegateStr) - 1;
+	FunctionData.MarshallAndCallName = FString::Printf( TEXT( "F%s_DelegateWrapper" ), *ShortName );
+
+	// Setup delegate parameter
+	const FString ExtraParam = FString::Printf(
+		TEXT( "const %s& %s" ),
+		bIsMulticastDelegate ? TEXT( "FMulticastScriptDelegate" ) : TEXT( "FScriptDelegate" ),
+		*DelegateName
+	);
+
+	DelegateOutput.Log(TEXT("static "));
+
+	// export the line that looks like: int32 Main(const FString& Parms)
+	ExportNativeFunctionHeader(DelegateOutput, OutReferenceGatherers.ForwardDeclarations, FunctionData, EExportFunctionType::Event, EExportFunctionHeaderStyle::Declaration, *ExtraParam, *GetAPIString());
+
+	FParmsAndReturnProperties Parameters = GetFunctionParmsAndReturn(FunctionData.FunctionReference);
+
+	WriteEventFunctionPrologue(DelegateOutput, 0, Parameters, Function->GetOuter(), *DelegateName);
+	{
+		const TCHAR* DelegateType = bIsMulticastDelegate ? TEXT( "ProcessMulticastDelegate" ) : TEXT( "ProcessDelegate" );
+		const TCHAR* DelegateArg  = Parameters.HasParms() ? TEXT("&Parms") : TEXT("NULL");
+		DelegateOutput.Logf(TEXT("\t%s.%s<UObject>(%s);\r\n"), *DelegateName, DelegateType, DelegateArg);
+	}
+	WriteEventFunctionEpilogue(DelegateOutput, 0, Parameters);
+
+	FString MacroName = SourceFile.GetGeneratedMacroName(FunctionData.MacroLine, TEXT("_DELEGATE"));
+	WriteMacro(Out, MacroName, DelegateOutput);
+}
+
+void FNativeClassHeaderGenerator::ExportEventParm(FUHTStringBuilder& Out, TSet<FString>& PropertyFwd, UFunction* Function, int32 Indent, bool bOutputConstructor, EExportingState ExportingState)
+{
+	if (!WillExportEventParms(Function))
+	{
+		return;
+	}
+
+	FString FunctionName = Function->GetName();
+	if (Function->HasAnyFunctionFlags(FUNC_Delegate))
+	{
+		FunctionName.LeftChopInline(HEADER_GENERATED_DELEGATE_SIGNATURE_SUFFIX_LENGTH, false);
+	}
+
+	FString EventParmStructName = GetEventStructParamsName(Function->GetOuter(), *FunctionName);
+	Out.Logf(TEXT("%sstruct %s\r\n"), FCString::Tab(Indent), *EventParmStructName);
+	Out.Logf(TEXT("%s{\r\n"), FCString::Tab(Indent));
+
+	for (FProperty* Prop : TFieldRange<FProperty>(Function))
+	{
+		if (!(Prop->PropertyFlags & CPF_Parm))
+		{
+			continue;
+		}
+
+		PropertyFwd.Add(Prop->GetCPPTypeForwardDeclaration());
+
+		FUHTStringBuilder PropertyText;
+		PropertyText.Log(FCString::Tab(Indent + 1));
+
+		bool bEmitConst = Prop->HasAnyPropertyFlags(CPF_ConstParm) && Prop->IsA<FObjectProperty>();
+
+		//@TODO: UCREMOVAL: This is awful code duplication to avoid a double-const
+		{
+			// export 'const' for parameters
+			const bool bIsConstParam = (Prop->IsA(FInterfaceProperty::StaticClass()) && !Prop->HasAllPropertyFlags(CPF_OutParm)); //@TODO: This should be const once that flag exists
+			const bool bIsOnConstClass = (Prop->IsA(FObjectProperty::StaticClass()) && ((FObjectProperty*)Prop)->PropertyClass != NULL && ((FObjectProperty*)Prop)->PropertyClass->HasAnyClassFlags(CLASS_Const));
+
+			if (bIsConstParam || bIsOnConstClass)
+			{
+				bEmitConst = false; // ExportCppDeclaration will do it for us
+			}
+		}
+
+		if (bEmitConst)
+		{
+			PropertyText.Logf(TEXT("const "));
+		}
+
+		const FString* Dim = GArrayDimensions.Find(Prop);
+		Prop->ExportCppDeclaration(PropertyText, EExportedDeclaration::Local, Dim ? **Dim : NULL);
+		ApplyAlternatePropertyExportText(Prop, PropertyText, ExportingState);
+
+		PropertyText.Log(TEXT(";\r\n"));
+		Out += *PropertyText;
+
+	}
+	// constructor must initialize the return property if it needs it
+	FProperty* Prop = Function->GetReturnProperty();
+	if (Prop && bOutputConstructor)
+	{
+		FUHTStringBuilder InitializationAr;
+
+		FStructProperty* InnerStruct = CastField<FStructProperty>(Prop);
+		bool bNeedsOutput = true;
+		if (InnerStruct)
+		{
+			bNeedsOutput = InnerStruct->HasNoOpConstructor();
+		}
+		else if (
+			CastField<FNameProperty>(Prop) ||
+			CastField<FDelegateProperty>(Prop) ||
+			CastField<FMulticastDelegateProperty>(Prop) ||
+			CastField<FStrProperty>(Prop) ||
+			CastField<FTextProperty>(Prop) ||
+			CastField<FArrayProperty>(Prop) ||
+			CastField<FMapProperty>(Prop) ||
+			CastField<FSetProperty>(Prop) ||
+			CastField<FInterfaceProperty>(Prop) ||
+			CastField<FFieldPathProperty>(Prop)
+			)
+		{
+			bNeedsOutput = false;
+		}
+		if (bNeedsOutput)
+		{
+			check(Prop->ArrayDim == 1); // can't return arrays
+			Out.Logf(TEXT("\r\n%s/** Constructor, initializes return property only **/\r\n"), FCString::Tab(Indent + 1));
+			Out.Logf(TEXT("%s%s()\r\n"), FCString::Tab(Indent + 1), *EventParmStructName);
+			Out.Logf(TEXT("%s%s %s(%s)\r\n"), FCString::Tab(Indent + 2), TEXT(":"), *Prop->GetName(), *GetNullParameterValue(Prop, true));
+			Out.Logf(TEXT("%s{\r\n"), FCString::Tab(Indent + 1));
+			Out.Logf(TEXT("%s}\r\n"), FCString::Tab(Indent + 1));
+		}
+	}
+	Out.Logf(TEXT("%s};\r\n"), FCString::Tab(Indent));
+}
+
+/**
+ * Get the intrinsic null value for this property
+ *
+ * @param	Prop				the property to get the null value for
+ * @param	bMacroContext		true when exporting the P_GET* macro, false when exporting the friendly C++ function header
+ *
+ * @return	the intrinsic null value for the property (0 for ints, TEXT("") for strings, etc.)
+ */
+FString FNativeClassHeaderGenerator::GetNullParameterValue( FProperty* Prop, bool bInitializer/*=false*/ )
+{
+	FFieldClass* PropClass = Prop->GetClass();
+	FObjectPropertyBase* ObjectProperty = CastField<FObjectPropertyBase>(Prop);
+	if (PropClass == FByteProperty::StaticClass())
+	{
+		FByteProperty* ByteProp = (FByteProperty*)Prop;
+
+		// if it's an enum class then we need an explicit cast
+		if( ByteProp->Enum && ByteProp->Enum->GetCppForm() == UEnum::ECppForm::EnumClass )
+		{
+			return FString::Printf(TEXT("(%s)0"), *ByteProp->GetCPPType());
+		}
+
+		return TEXT("0");
+	}
+	else if (PropClass == FEnumProperty::StaticClass())
+	{
+		FEnumProperty* EnumProp = (FEnumProperty*)Prop;
+
+		return FString::Printf(TEXT("(%s)0"), *EnumProp->Enum->GetName());
+	}
+	else if ( PropClass == FBoolProperty::StaticClass() )
+	{
+		return TEXT("false");
+	}
+	else if ( PropClass == FIntProperty::StaticClass()
+	||	PropClass == FFloatProperty::StaticClass()
+	||	PropClass == FDoubleProperty::StaticClass())
+	{
+		return TEXT("0");
+	}
+	else if ( PropClass == FNameProperty::StaticClass() )
+	{
+		return TEXT("NAME_None");
+	}
+	else if ( PropClass == FStrProperty::StaticClass() )
+	{
+		return TEXT("TEXT(\"\")");
+	}
+	else if ( PropClass == FTextProperty::StaticClass() )
+	{
+		return TEXT("FText::GetEmpty()");
+	}
+	else if ( PropClass == FArrayProperty::StaticClass()
+		||    PropClass == FMapProperty::StaticClass()
+		||    PropClass == FSetProperty::StaticClass()
+		||    PropClass == FDelegateProperty::StaticClass()
+		||    PropClass == FMulticastDelegateProperty::StaticClass() )
+	{
+		FString Type, ExtendedType;
+		Type = Prop->GetCPPType(&ExtendedType,CPPF_OptionalValue);
+		return Type + ExtendedType + TEXT("()");
+	}
+	else if ( PropClass == FStructProperty::StaticClass() )
+	{
+		bool bHasNoOpConstuctor = CastFieldChecked<FStructProperty>(Prop)->HasNoOpConstructor();
+		if (bInitializer && bHasNoOpConstuctor)
+		{
+			return TEXT("ForceInit");
+		}
+
+		FString Type, ExtendedType;
+		Type = Prop->GetCPPType(&ExtendedType,CPPF_OptionalValue);
+		return Type + ExtendedType + (bHasNoOpConstuctor ? TEXT("(ForceInit)") : TEXT("()"));
+	}
+	else if (ObjectProperty)
+	{
+		return TEXT("NULL");
+	}
+	else if ( PropClass == FInterfaceProperty::StaticClass() )
+	{
+		return TEXT("NULL");
+	}
+	else if (PropClass == FFieldPathProperty::StaticClass())
+	{
+		return TEXT("nullptr");
+	}
+
+	UE_LOG(LogCompile, Fatal,TEXT("GetNullParameterValue - Unhandled property type '%s': %s"), *Prop->GetClass()->GetName(), *Prop->GetPathName());
+	return TEXT("");
+}
+
+
+FString FNativeClassHeaderGenerator::GetFunctionReturnString(UFunction* Function, FReferenceGatherers& OutReferenceGatherers)
+{
+	FString Result;
+
+	if (FProperty* Return = Function->GetReturnProperty())
+	{
+		FString ExtendedReturnType;
+		OutReferenceGatherers.ForwardDeclarations.Add(Return->GetCPPTypeForwardDeclaration());
+		FString ReturnType = Return->GetCPPType(&ExtendedReturnType, CPPF_ArgumentOrReturnValue);
+		FUHTStringBuilder ReplacementText;
+		ReplacementText += MoveTemp(ReturnType);
+		ApplyAlternatePropertyExportText(Return, ReplacementText, EExportingState::Normal);
+		Result = MoveTemp(ReplacementText) + MoveTemp(ExtendedReturnType);
+	}
+	else
+	{
+		Result = TEXT("void");
+	}
+
+	return Result;
+}
+
+/**
+ * Converts Position within File to Line and Column.
+ *
+ * @param File File contents.
+ * @param Position Position in string to convert.
+ * @param OutLine Result line.
+ * @param OutColumn Result column.
+ */
+void GetLineAndColumnFromPositionInFile(const FString& File, int32 Position, int32& OutLine, int32& OutColumn)
+{
+	OutLine = 1;
+	OutColumn = 1;
+
+	int32 i;
+	for (i = 1; i <= Position; ++i)
+	{
+		if (File[i] == '\n')
+		{
+			++OutLine;
+			OutColumn = 0;
+		}
+		else
+		{
+			++OutColumn;
+		}
+	}
+}
+
+bool FNativeClassHeaderGenerator::IsMissingVirtualSpecifier(const FString& SourceFile, int32 FunctionNamePosition)
+{
+	auto IsEndOfSearchChar = [](TCHAR C) { return (C == TEXT('}')) || (C == TEXT('{')) || (C == TEXT(';')); };
+
+	// Find first occurrence of "}", ";", "{" going backwards from ImplementationPosition.
+	int32 EndOfSearchCharIndex = SourceFile.FindLastCharByPredicate(IsEndOfSearchChar, FunctionNamePosition);
+	check(EndOfSearchCharIndex != INDEX_NONE);
+
+	// Then find if there is "virtual" keyword starting from position of found character to ImplementationPosition
+	return !HasIdentifierExactMatch(&SourceFile[EndOfSearchCharIndex], &SourceFile[FunctionNamePosition], TEXT("virtual"));
+}
+
+FString CreateClickableErrorMessage(const FString& Filename, int32 Line, int32 Column)
+{
+	return FString::Printf(TEXT("%s(%d,%d): error: "), *Filename, Line, Column);
+}
+
+void FNativeClassHeaderGenerator::CheckRPCFunctions(FReferenceGatherers& OutReferenceGatherers, const FFuncInfo& FunctionData, const FString& ClassName, int32 ImplementationPosition, int32 ValidatePosition, const FUnrealSourceFile& SourceFile) const
+{
+	bool bHasImplementation = ImplementationPosition != INDEX_NONE;
+	bool bHasValidate = ValidatePosition != INDEX_NONE;
+
+	UFunction* Function = FunctionData.FunctionReference;
+	FString FunctionReturnType = GetFunctionReturnString(Function, OutReferenceGatherers);
+	const TCHAR* ConstModifier = (Function->HasAllFunctionFlags(FUNC_Const) ? TEXT("const ") : TEXT(" "));
+
+	const bool bIsNative = Function->HasAllFunctionFlags(FUNC_Native);
+	const bool bIsNet = Function->HasAllFunctionFlags(FUNC_Net);
+	const bool bIsNetValidate = Function->HasAllFunctionFlags(FUNC_NetValidate);
+	const bool bIsNetResponse = Function->HasAllFunctionFlags(FUNC_NetResponse);
+	const bool bIsBlueprintEvent = Function->HasAllFunctionFlags(FUNC_BlueprintEvent);
+
+	bool bNeedsImplementation = (bIsNet && !bIsNetResponse) || bIsBlueprintEvent || bIsNative;
+	bool bNeedsValidate = (bIsNative || bIsNet) && !bIsNetResponse && bIsNetValidate;
+
+	check(bNeedsImplementation || bNeedsValidate);
+
+	FString ParameterString = GetFunctionParameterString(Function, OutReferenceGatherers);
+	const FString& Filename = SourceFile.GetFilename();
+	const FString& FileContent = SourceFile.GetContent();
+
+	//
+	// Get string with function specifiers, listing why we need _Implementation or _Validate functions.
+	//
+	TArray<const TCHAR*, TInlineAllocator<4>> FunctionSpecifiers;
+	if (bIsNative)			{ FunctionSpecifiers.Add(TEXT("Native"));			}
+	if (bIsNet)				{ FunctionSpecifiers.Add(TEXT("Net"));				}
+	if (bIsBlueprintEvent)	{ FunctionSpecifiers.Add(TEXT("BlueprintEvent"));	}
+	if (bIsNetValidate)		{ FunctionSpecifiers.Add(TEXT("NetValidate"));		}
+
+	check(FunctionSpecifiers.Num() > 0);
+
+	//
+	// Coin static_assert message
+	//
+	FUHTStringBuilder AssertMessage;
+	AssertMessage.Logf(TEXT("Function %s was marked as %s"), *(Function->GetName()), FunctionSpecifiers[0]);
+	for (int32 i = 1; i < FunctionSpecifiers.Num(); ++i)
+	{
+		AssertMessage.Logf(TEXT(", %s"), FunctionSpecifiers[i]);
+	}
+
+	AssertMessage.Logf(TEXT("."));
+
+	//
+	// Check if functions are missing.
+	//
+	int32 Line;
+	int32 Column;
+	GetLineAndColumnFromPositionInFile(FileContent, FunctionData.InputPos, Line, Column);
+	if (bNeedsImplementation && !bHasImplementation)
+	{
+		FString ErrorPosition = CreateClickableErrorMessage(Filename, Line, Column);
+		FString FunctionDecl = FString::Printf(TEXT("virtual %s %s::%s(%s) %s"), *FunctionReturnType, *ClassName, *FunctionData.CppImplName, *ParameterString, *ConstModifier);
+		FError::Throwf(TEXT("%s%s Declare function %s"), *ErrorPosition, *AssertMessage, *FunctionDecl);
+	}
+
+	if (bNeedsValidate && !bHasValidate)
+	{
+		FString ErrorPosition = CreateClickableErrorMessage(Filename, Line, Column);
+		FString FunctionDecl = FString::Printf(TEXT("virtual bool %s::%s(%s) %s"), *ClassName, *FunctionData.CppValidationImplName, *ParameterString, *ConstModifier);
+		FError::Throwf(TEXT("%s%s Declare function %s"), *ErrorPosition, *AssertMessage, *FunctionDecl);
+	}
+
+	//
+	// If all needed functions are declared, check if they have virtual specifiers.
+	//
+	if (bNeedsImplementation && bHasImplementation && IsMissingVirtualSpecifier(FileContent, ImplementationPosition))
+	{
+		GetLineAndColumnFromPositionInFile(FileContent, ImplementationPosition, Line, Column);
+		FString ErrorPosition = CreateClickableErrorMessage(Filename, Line, Column);
+		FString FunctionDecl = FString::Printf(TEXT("%s %s::%s(%s) %s"), *FunctionReturnType, *ClassName, *FunctionData.CppImplName, *ParameterString, *ConstModifier);
+		FError::Throwf(TEXT("%sDeclared function %sis not marked as virtual."), *ErrorPosition, *FunctionDecl);
+	}
+
+	if (bNeedsValidate && bHasValidate && IsMissingVirtualSpecifier(FileContent, ValidatePosition))
+	{
+		GetLineAndColumnFromPositionInFile(FileContent, ValidatePosition, Line, Column);
+		FString ErrorPosition = CreateClickableErrorMessage(Filename, Line, Column);
+		FString FunctionDecl = FString::Printf(TEXT("bool %s::%s(%s) %s"), *ClassName, *FunctionData.CppValidationImplName, *ParameterString, *ConstModifier);
+		FError::Throwf(TEXT("%sDeclared function %sis not marked as virtual."), *ErrorPosition, *FunctionDecl);
+	}
+}
+
+void FNativeClassHeaderGenerator::ExportNativeFunctionHeader(
+	FOutputDevice&                   Out,
+	TSet<FString>&                   OutFwdDecls,
+	const FFuncInfo&                 FunctionData,
+	EExportFunctionType::Type        FunctionType,
+	EExportFunctionHeaderStyle::Type FunctionHeaderStyle,
+	const TCHAR*                     ExtraParam,
+	const TCHAR*                     APIString
+)
+{
+	UFunction* Function = FunctionData.FunctionReference;
+
+	const bool bIsDelegate   = Function->HasAnyFunctionFlags( FUNC_Delegate );
+	const bool bIsInterface  = !bIsDelegate && Function->GetOwnerClass()->HasAnyClassFlags(CLASS_Interface);
+	const bool bIsK2Override = Function->HasAnyFunctionFlags( FUNC_BlueprintEvent );
+
+	if (!bIsDelegate)
+	{
+		Out.Log(TEXT("\t"));
+	}
+
+	if (FunctionHeaderStyle == EExportFunctionHeaderStyle::Declaration)
+	{
+		// cpp implementation of functions never have these appendages
+
+		// If the function was marked as 'RequiredAPI', then add the *_API macro prefix.  Note that if the class itself
+		// was marked 'RequiredAPI', this is not needed as C++ will exports all methods automatically.
+		if (FunctionType != EExportFunctionType::Event &&
+			!Function->GetOwnerClass()->HasAnyClassFlags(CLASS_RequiredAPI) &&
+			(FunctionData.FunctionExportFlags & FUNCEXPORT_RequiredAPI))
+		{
+			Out.Log(APIString);
+		}
+
+		if(FunctionType == EExportFunctionType::Interface)
+		{
+			Out.Log(TEXT("static "));
+		}
+		else if (bIsK2Override)
+		{
+			Out.Log(TEXT("virtual "));
+		}
+		// if the owning class is an interface class
+		else if ( bIsInterface )
+		{
+			Out.Log(TEXT("virtual "));
+		}
+		// this is not an event, the function is not a static function and the function is not marked final
+		else if ( FunctionType != EExportFunctionType::Event && !Function->HasAnyFunctionFlags(FUNC_Static) && !(FunctionData.FunctionExportFlags & FUNCEXPORT_Final) )
+		{
+			Out.Log(TEXT("virtual "));
+		}
+		else if( FunctionData.FunctionExportFlags & FUNCEXPORT_Inline )
+		{
+			Out.Log(TEXT("inline "));
+		}
+	}
+
+	FProperty* ReturnProperty = Function->GetReturnProperty();
+	if (ReturnProperty != nullptr)
+	{
+		if (ReturnProperty->HasAnyPropertyFlags(EPropertyFlags::CPF_ConstParm))
+		{
+			Out.Log(TEXT("const "));
+		}
+
+		FString ExtendedReturnType;
+		FString ReturnType = ReturnProperty->GetCPPType(&ExtendedReturnType, (FunctionHeaderStyle == EExportFunctionHeaderStyle::Definition && (FunctionType != EExportFunctionType::Interface) ? CPPF_Implementation : 0) | CPPF_ArgumentOrReturnValue);
+		OutFwdDecls.Add(ReturnProperty->GetCPPTypeForwardDeclaration());
+		FUHTStringBuilder ReplacementText;
+		ReplacementText += ReturnType;
+		ApplyAlternatePropertyExportText(ReturnProperty, ReplacementText, EExportingState::Normal);
+		Out.Logf(TEXT("%s%s"), *ReplacementText, *ExtendedReturnType);
+	}
+	else
+	{
+		Out.Log( TEXT("void") );
+	}
+
+	FString FunctionName;
+	if (FunctionHeaderStyle == EExportFunctionHeaderStyle::Definition)
+	{
+		FunctionName = FString::Printf(TEXT("%s::"), *FNameLookupCPP::GetNameCPP(CastChecked<UClass>(Function->GetOuter()), bIsInterface || FunctionType == EExportFunctionType::Interface));
+	}
+
+	if (FunctionType == EExportFunctionType::Interface)
+	{
+		FunctionName += FString::Printf(TEXT("Execute_%s"), *Function->GetName());
+	}
+	else if (FunctionType == EExportFunctionType::Event)
+	{
+		FunctionName += FunctionData.MarshallAndCallName;
+	}
+	else
+	{
+		FunctionName += FunctionData.CppImplName;
+	}
+
+	Out.Logf(TEXT(" %s("), *FunctionName);
+
+	int32 ParmCount=0;
+
+	// Emit extra parameter if we have one
+	if( ExtraParam )
+	{
+		Out.Logf(TEXT("%s"), ExtraParam);
+		++ParmCount;
+	}
+
+	for (FProperty* Property : TFieldRange<FProperty>(Function))
+	{
+		if ((Property->PropertyFlags & (CPF_Parm | CPF_ReturnParm)) != CPF_Parm)
+		{
+			continue;
+		}
+
+		OutFwdDecls.Add(Property->GetCPPTypeForwardDeclaration());
+
+		if( ParmCount++ )
+		{
+			Out.Log(TEXT(", "));
+		}
+
+		FUHTStringBuilder PropertyText;
+
+		const FString* Dim = GArrayDimensions.Find(Property);
+		Property->ExportCppDeclaration( PropertyText, EExportedDeclaration::Parameter, Dim ? **Dim : NULL );
+		ApplyAlternatePropertyExportText(Property, PropertyText, EExportingState::Normal);
+
+		Out.Logf(TEXT("%s"), *PropertyText);
+	}
+
+	Out.Log( TEXT(")") );
+	if (FunctionType != EExportFunctionType::Interface)
+	{
+		if (!bIsDelegate && Function->HasAllFunctionFlags(FUNC_Const))
+		{
+			Out.Log( TEXT(" const") );
+		}
+
+		if (bIsInterface && FunctionHeaderStyle == EExportFunctionHeaderStyle::Declaration)
+		{
+			// all methods in interface classes are pure virtuals
+			if (bIsK2Override)
+			{
+				// For BlueprintNativeEvent methods we emit a stub implementation. This allows Blueprints that implement the interface class to be nativized.
+				FString ReturnValue;
+				if (ReturnProperty != nullptr)
+				{
+					FByteProperty* ByteProperty = CastField<FByteProperty>(ReturnProperty);
+					if (ByteProperty != nullptr && ByteProperty->Enum != nullptr && ByteProperty->Enum->GetCppForm() != UEnum::ECppForm::EnumClass)
+					{
+						ReturnValue = FString::Printf(TEXT(" return TEnumAsByte<%s>(%s); "), *ByteProperty->Enum->CppType, *GetNullParameterValue(ReturnProperty, false));
+					}
+					else
+					{
+						ReturnValue = FString::Printf(TEXT(" return %s; "), *GetNullParameterValue(ReturnProperty, false));
+					}
+				}
+
+				Out.Logf(TEXT(" {%s}"), *ReturnValue);
+			}
+			else
+			{
+				Out.Log(TEXT("=0"));
+			}
+		}
+	}
+}
+
+/**
+ * Export the actual internals to a standard thunk function
+ *
+ * @param RPCWrappers output device for writing
+ * @param FunctionData function data for the current function
+ * @param Parameters list of parameters in the function
+ * @param Return return parameter for the function
+ * @param DeprecationWarningOutputDevice Device to output deprecation warnings for _Validate and _Implementation functions.
+ */
+void FNativeClassHeaderGenerator::ExportFunctionThunk(FUHTStringBuilder& RPCWrappers, FReferenceGatherers& OutReferenceGatherers, UFunction* Function, const FFuncInfo& FunctionData, const TArray<FProperty*>& Parameters, FProperty* Return) const
+{
+	// export the GET macro for this parameter
+	FString ParameterList;
+	for (int32 ParameterIndex = 0; ParameterIndex < Parameters.Num(); ParameterIndex++)
+	{
+		FProperty* Param = Parameters[ParameterIndex];
+		OutReferenceGatherers.ForwardDeclarations.Add(Param->GetCPPTypeForwardDeclaration());
+
+		FString EvalBaseText = TEXT("P_GET_");	// e.g. P_GET_STR
+		FString EvalModifierText;				// e.g. _REF
+		FString EvalParameterText;				// e.g. (UObject*,NULL)
+
+		FString TypeText;
+
+		if (Param->ArrayDim > 1)
+		{
+			EvalBaseText += TEXT("ARRAY");
+			TypeText = Param->GetCPPType();
+		}
+		else
+		{
+			EvalBaseText += Param->GetCPPMacroType(TypeText);
+
+			FArrayProperty* ArrayProperty = CastField<FArrayProperty>(Param);
+			if (ArrayProperty)
+			{
+				FInterfaceProperty* InterfaceProperty = CastField<FInterfaceProperty>(ArrayProperty->Inner);
+				if (InterfaceProperty)
+				{
+					FString InterfaceTypeText;
+					InterfaceProperty->GetCPPMacroType(InterfaceTypeText);
+					TypeText += FString::Printf(TEXT("<%s>"), *InterfaceTypeText);
+				}
+			}
+		}
+
+		bool bPassAsNoPtr = Param->HasAllPropertyFlags(CPF_UObjectWrapper | CPF_OutParm) && Param->IsA(FClassProperty::StaticClass());
+		if (bPassAsNoPtr)
+		{
+			TypeText = Param->GetCPPType();
+		}
+
+		FUHTStringBuilder ReplacementText;
+		ReplacementText += TypeText;
+
+		ApplyAlternatePropertyExportText(Param, ReplacementText, EExportingState::Normal);
+		TypeText = ReplacementText;
+
+		FString DefaultValueText;
+		FString ParamPrefix = TEXT("Z_Param_");
+
+		// if this property is an out parm, add the REF tag
+		if (Param->PropertyFlags & CPF_OutParm)
+		{
+			if (!bPassAsNoPtr)
+			{
+				EvalModifierText += TEXT("_REF");
+			}
+			else
+			{
+				// Parameters passed as TSubclassOf<Class>& shouldn't have asterisk added.
+				EvalModifierText += TEXT("_REF_NO_PTR");
+			}
+
+			ParamPrefix += TEXT("Out_");
+		}
+
+		// if this property requires a specialization, add a comma to the type name so we can print it out easily
+		if (TypeText != TEXT(""))
+		{
+			TypeText += TCHAR(',');
+		}
+
+		FString ParamName = ParamPrefix + Param->GetName();
+
+		EvalParameterText = FString::Printf(TEXT("(%s%s%s)"), *TypeText, *ParamName, *DefaultValueText);
+
+		RPCWrappers.Logf(TEXT("\t\t%s%s%s;") LINE_TERMINATOR, *EvalBaseText, *EvalModifierText, *EvalParameterText);
+
+		// add this property to the parameter list string
+		if (ParameterList.Len())
+		{
+			ParameterList += TCHAR(',');
+		}
+
+		{
+			FDelegateProperty* DelegateProp = CastField< FDelegateProperty >(Param);
+			if (DelegateProp != NULL)
+			{
+				// For delegates, add an explicit conversion to the specific type of delegate before passing it along
+				const FString FunctionName = DelegateProp->SignatureFunction->GetName().LeftChop(HEADER_GENERATED_DELEGATE_SIGNATURE_SUFFIX_LENGTH);
+				ParamName = FString::Printf(TEXT("F%s(%s)"), *FunctionName, *ParamName);
+			}
+		}
+
+		{
+			FMulticastDelegateProperty* MulticastDelegateProp = CastField< FMulticastDelegateProperty >(Param);
+			if (MulticastDelegateProp != NULL)
+			{
+				// For delegates, add an explicit conversion to the specific type of delegate before passing it along
+				const FString FunctionName = MulticastDelegateProp->SignatureFunction->GetName().LeftChop(HEADER_GENERATED_DELEGATE_SIGNATURE_SUFFIX_LENGTH);
+				ParamName = FString::Printf(TEXT("F%s(%s)"), *FunctionName, *ParamName);
+			}
+		}
+
+		UEnum* Enum = nullptr;
+		FByteProperty* ByteProp = CastField<FByteProperty>(Param);
+		if (ByteProp && ByteProp->Enum)
+		{
+			Enum = ByteProp->Enum;
+		}
+		else if (Param->IsA<FEnumProperty>())
+		{
+			Enum = ((FEnumProperty*)Param)->Enum;
+		}
+
+		if (Enum)
+		{
+			// For enums, add an explicit conversion
+			if (!(Param->PropertyFlags & CPF_OutParm))
+			{
+				ParamName = FString::Printf(TEXT("%s(%s)"), *Enum->CppType, *ParamName);
+			}
+			else
+			{
+				if (Enum->GetCppForm() == UEnum::ECppForm::EnumClass)
+				{
+					// If we're an enum class don't require the wrapper
+					ParamName = FString::Printf(TEXT("(%s&)(%s)"), *Enum->CppType, *ParamName);
+				}
+				else
+				{
+					ParamName = FString::Printf(TEXT("(TEnumAsByte<%s>&)(%s)"), *Enum->CppType, *ParamName);
+				}
+			}
+		}
+
+		ParameterList += ParamName;
+	}
+
+	RPCWrappers += TEXT("\t\tP_FINISH;") LINE_TERMINATOR;
+	RPCWrappers += TEXT("\t\tP_NATIVE_BEGIN;") LINE_TERMINATOR;
+
+	ClassDefinitionRange ClassRange;
+	if (ClassDefinitionRanges.Contains(Function->GetOwnerClass()))
+	{
+		ClassRange = ClassDefinitionRanges[Function->GetOwnerClass()];
+		ClassRange.Validate();
+	}
+
+	const TCHAR* ClassStart = ClassRange.Start;
+	const TCHAR* ClassEnd   = ClassRange.End;
+	FString      ClassName  = Function->GetOwnerClass()->GetName();
+
+	FString ClassDefinition(UE_PTRDIFF_TO_INT32(ClassEnd - ClassStart), ClassStart);
+
+	bool bHasImplementation = HasIdentifierExactMatch(ClassDefinition, FunctionData.CppImplName);
+	bool bHasValidate = HasIdentifierExactMatch(ClassDefinition, FunctionData.CppValidationImplName);
+
+	bool bShouldEnableImplementationDeprecation =
+		// Enable deprecation warnings only if GENERATED_BODY is used inside class or interface (not GENERATED_UCLASS_BODY etc.)
+		ClassRange.bHasGeneratedBody
+		// and implementation function is called, but not the one declared by user
+		&& (FunctionData.CppImplName != Function->GetName() && !bHasImplementation);
+
+	bool bShouldEnableValidateDeprecation =
+		// Enable deprecation warnings only if GENERATED_BODY is used inside class or interface (not GENERATED_UCLASS_BODY etc.)
+		ClassRange.bHasGeneratedBody
+		// and validation function is called
+		&& (FunctionData.FunctionFlags & FUNC_NetValidate) && !bHasValidate;
+
+	//Emit warning here if necessary
+	FUHTStringBuilder FunctionDeclaration;
+	ExportNativeFunctionHeader(FunctionDeclaration, OutReferenceGatherers.ForwardDeclarations, FunctionData, EExportFunctionType::Function, EExportFunctionHeaderStyle::Declaration, nullptr, *GetAPIString());
+
+	// Call the validate function if there is one
+	if (!(FunctionData.FunctionExportFlags & FUNCEXPORT_CppStatic) && (FunctionData.FunctionFlags & FUNC_NetValidate))
+	{
+		RPCWrappers.Logf(TEXT("\t\tif (!P_THIS->%s(%s))") LINE_TERMINATOR, *FunctionData.CppValidationImplName, *ParameterList);
+		RPCWrappers.Logf(TEXT("\t\t{") LINE_TERMINATOR);
+		RPCWrappers.Logf(TEXT("\t\t\tRPC_ValidateFailed(TEXT(\"%s\"));") LINE_TERMINATOR, *FunctionData.CppValidationImplName);
+		RPCWrappers.Logf(TEXT("\t\t\treturn;") LINE_TERMINATOR);	// If we got here, the validation function check failed
+		RPCWrappers.Logf(TEXT("\t\t}") LINE_TERMINATOR);
+	}
+
+	// write out the return value
+	RPCWrappers.Log(TEXT("\t\t"));
+	if (Return)
+	{
+		OutReferenceGatherers.ForwardDeclarations.Add(Return->GetCPPTypeForwardDeclaration());
+
+		FUHTStringBuilder ReplacementText;
+		FString ReturnExtendedType;
+		ReplacementText += Return->GetCPPType(&ReturnExtendedType);
+		ApplyAlternatePropertyExportText(Return, ReplacementText, EExportingState::Normal);
+
+		FString ReturnType = ReplacementText;
+		RPCWrappers.Logf(TEXT("*(%s%s*)") TEXT(PREPROCESSOR_TO_STRING(RESULT_PARAM)) TEXT("="), *ReturnType, *ReturnExtendedType);
+	}
+
+	// export the call to the C++ version
+	if (FunctionData.FunctionExportFlags & FUNCEXPORT_CppStatic)
+	{
+		RPCWrappers.Logf(TEXT("%s::%s(%s);") LINE_TERMINATOR, *FNameLookupCPP::GetNameCPP(Function->GetOwnerClass()), *FunctionData.CppImplName, *ParameterList);
+	}
+	else
+	{
+		RPCWrappers.Logf(TEXT("P_THIS->%s(%s);") LINE_TERMINATOR, *FunctionData.CppImplName, *ParameterList);
+	}
+	RPCWrappers += TEXT("\t\tP_NATIVE_END;") LINE_TERMINATOR;
+}
+
+FString FNativeClassHeaderGenerator::GetFunctionParameterString(UFunction* Function, FReferenceGatherers& OutReferenceGatherers)
+{
+	FString ParameterList;
+	FUHTStringBuilder PropertyText;
+
+	for (FProperty* Property : TFieldRange<FProperty>(Function))
+	{
+		OutReferenceGatherers.ForwardDeclarations.Add(Property->GetCPPTypeForwardDeclaration());
+
+		if ((Property->PropertyFlags & (CPF_Parm | CPF_ReturnParm)) != CPF_Parm)
+		{
+			break;
+		}
+
+		if (ParameterList.Len())
+		{
+			ParameterList += TEXT(", ");
+		}
+
+		FString* Dim = GArrayDimensions.Find(Property);
+		Property->ExportCppDeclaration(PropertyText, EExportedDeclaration::Parameter, Dim ? **Dim : nullptr, 0, true);
+		ApplyAlternatePropertyExportText(Property, PropertyText, EExportingState::Normal);
+
+		ParameterList += PropertyText;
+		PropertyText.Reset();
+	}
+
+	return ParameterList;
+}
+
+struct FNativeFunctionStringBuilder
+{
+	FUHTStringBuilder RPCWrappers;
+	FUHTStringBuilder RPCImplementations;
+	FUHTStringBuilder AutogeneratedBlueprintFunctionDeclarations;
+	FUHTStringBuilder AutogeneratedBlueprintFunctionDeclarationsOnlyNotDeclared;
+	FUHTStringBuilder AutogeneratedStaticData;
+	FUHTStringBuilder AutogeneratedStaticDataFuncs;
+};
+
+void FNativeClassHeaderGenerator::ExportNativeFunctions(FOutputDevice& OutGeneratedHeaderText, FOutputDevice& OutGeneratedCPPText, FOutputDevice& OutMacroCalls, FOutputDevice& OutNoPureDeclsMacroCalls, FReferenceGatherers& OutReferenceGatherers, const FUnrealSourceFile& SourceFile, UClass* Class, FClassMetaData* ClassData) const
+{
+	FNativeFunctionStringBuilder RuntimeStringBuilders;
+	FNativeFunctionStringBuilder EditorStringBuilders;
+
+	const FString ClassCPPName = FNameLookupCPP::GetNameCPP(Class, Class->HasAnyClassFlags(CLASS_Interface));
+
+	ClassDefinitionRange ClassRange;
+	if (ClassDefinitionRanges.Contains(Class))
+	{
+		ClassRange = ClassDefinitionRanges[Class];
+		ClassRange.Validate();
+	}
+
+	// gather static class data
+	TArray<FString> SparseClassDataTypes;
+	((FClass*)Class)->GetSparseClassDataTypes(SparseClassDataTypes);
+	FString FullClassName = ((FClass*)Class)->GetNameWithPrefix();
+	for (const FString& SparseClassDataString : SparseClassDataTypes)
+	{
+		RuntimeStringBuilders.AutogeneratedStaticData.Logf(TEXT("F%s* Get%s()\r\n"), *SparseClassDataString, *SparseClassDataString);
+		RuntimeStringBuilders.AutogeneratedStaticData += TEXT("{\r\n");
+		RuntimeStringBuilders.AutogeneratedStaticData.Logf(TEXT("\treturn (F%s*)(GetClass()->GetOrCreateSparseClassData());\r\n"), *SparseClassDataString);
+		RuntimeStringBuilders.AutogeneratedStaticData += TEXT("}\r\n");
+
+		RuntimeStringBuilders.AutogeneratedStaticData.Logf(TEXT("F%s* Get%s() const\r\n"), *SparseClassDataString, *SparseClassDataString);
+		RuntimeStringBuilders.AutogeneratedStaticData += TEXT("{\r\n");
+		RuntimeStringBuilders.AutogeneratedStaticData.Logf(TEXT("\treturn (F%s*)(GetClass()->GetOrCreateSparseClassData());\r\n"), *SparseClassDataString);
+		RuntimeStringBuilders.AutogeneratedStaticData += TEXT("}\r\n");
+
+		UScriptStruct* SparseClassDataStruct = FindObjectSafe<UScriptStruct>(ANY_PACKAGE, *SparseClassDataString);
+		while (SparseClassDataStruct != nullptr)
+		{
+			const FProperty* Child = CastField<FProperty>(SparseClassDataStruct->ChildProperties);
+			while (Child)
+			{
+				FString ReturnExtendedType;
+				FString VarType = Child->GetCPPType(&ReturnExtendedType, EPropertyExportCPPFlags::CPPF_ArgumentOrReturnValue | EPropertyExportCPPFlags::CPPF_Implementation);
+				if (!ReturnExtendedType.IsEmpty())
+				{
+					VarType.Append(ReturnExtendedType);
+				}
+				FString VarName = Child->GetName();
+				FString CleanVarName = VarName;
+				if (CastField<FBoolProperty>(Child) && VarName.StartsWith(TEXT("b"), ESearchCase::CaseSensitive))
+				{
+					CleanVarName = VarName.RightChop(1);
+				}
+
+				if (!Child->HasMetaData(NAME_NoGetter))
+				{
+					if (Child->HasMetaData(NAME_GetByRef))
+					{
+						RuntimeStringBuilders.AutogeneratedStaticDataFuncs.Logf(TEXT("const %s& Get%s()\r\n"), *VarType, *CleanVarName);
+					}
+					else
+					{
+						RuntimeStringBuilders.AutogeneratedStaticDataFuncs.Logf(TEXT("%s Get%s()\r\n"), *VarType, *CleanVarName);
+					}
+					RuntimeStringBuilders.AutogeneratedStaticDataFuncs.Logf(TEXT("{\r\n"));
+					RuntimeStringBuilders.AutogeneratedStaticDataFuncs.Logf(TEXT("\treturn Get%s()->%s;\r\n"), *SparseClassDataString, *VarName);
+					RuntimeStringBuilders.AutogeneratedStaticDataFuncs.Logf(TEXT("}\r\n"));
+
+					if (Child->HasMetaData(NAME_GetByRef))
+					{
+						RuntimeStringBuilders.AutogeneratedStaticDataFuncs.Logf(TEXT("const %s& Get%s() const\r\n"), *VarType, *CleanVarName);
+					}
+					else
+					{
+						RuntimeStringBuilders.AutogeneratedStaticDataFuncs.Logf(TEXT("%s Get%s() const\r\n"), *VarType, *CleanVarName);
+					}
+					RuntimeStringBuilders.AutogeneratedStaticDataFuncs.Logf(TEXT("{\r\n"));
+					RuntimeStringBuilders.AutogeneratedStaticDataFuncs.Logf(TEXT("\treturn Get%s()->%s;\r\n"), *SparseClassDataString, *VarName);
+					RuntimeStringBuilders.AutogeneratedStaticDataFuncs.Logf(TEXT("}\r\n"));
+				}
+
+				Child = CastField<FProperty>(Child->Next);
+			}
+
+			SparseClassDataStruct = Cast<UScriptStruct>(SparseClassDataStruct->GetSuperStruct());
+		}
+	}
+
+	// export the C++ stubs
+
+	for (UFunction* Function : TFieldRange<UFunction>(Class, EFieldIteratorFlags::ExcludeSuper))
+	{
+		if (!(Function->FunctionFlags & FUNC_Native))
+		{
+			continue;
+		}
+
+		const bool bEditorOnlyFunc = Function->HasAnyFunctionFlags(FUNC_EditorOnly);
+		FNativeFunctionStringBuilder& FuncStringBuilders = bEditorOnlyFunc ? EditorStringBuilders : RuntimeStringBuilders;
+
+		FFunctionData* CompilerInfo = FFunctionData::FindForFunction(Function);
+
+		const FFuncInfo& FunctionData = CompilerInfo->GetFunctionData();
+
+		// Custom thunks don't get any C++ stub function generated
+		if (FunctionData.FunctionExportFlags & FUNCEXPORT_CustomThunk)
+		{
+			continue;
+		}
+
+		// Should we emit these to RPC wrappers or just ignore them?
+		const bool bWillBeProgrammerTyped = FunctionData.CppImplName == Function->GetName();
+
+		if (!bWillBeProgrammerTyped)
+		{
+			const TCHAR* ClassStart = ClassRange.Start;
+			const TCHAR* ClassEnd   = ClassRange.End;
+			FString ClassDefinition(UE_PTRDIFF_TO_INT32(ClassEnd - ClassStart), ClassStart);
+
+			FString FunctionName = Function->GetName();
+			int32 ClassDefinitionStartPosition = UE_PTRDIFF_TO_INT32(ClassStart - *SourceFile.GetContent());
+
+			int32 ImplementationPosition = FindIdentifierExactMatch(ClassDefinition, FunctionData.CppImplName);
+			bool bHasImplementation = ImplementationPosition != INDEX_NONE;
+			if (bHasImplementation)
+			{
+				ImplementationPosition += ClassDefinitionStartPosition;
+			}
+
+			int32 ValidatePosition = FindIdentifierExactMatch(ClassDefinition, FunctionData.CppValidationImplName);
+			bool bHasValidate = ValidatePosition != INDEX_NONE;
+			if (bHasValidate)
+			{
+				ValidatePosition += ClassDefinitionStartPosition;
+			}
+
+			//Emit warning here if necessary
+			FUHTStringBuilder FunctionDeclaration;
+			ExportNativeFunctionHeader(FunctionDeclaration, OutReferenceGatherers.ForwardDeclarations, FunctionData, EExportFunctionType::Function, EExportFunctionHeaderStyle::Declaration, nullptr, *GetAPIString());
+			FunctionDeclaration.Log(TEXT(";\r\n"));
+
+			// Declare validation function if needed
+			if (FunctionData.FunctionFlags & FUNC_NetValidate)
+			{
+				FString ParameterList = GetFunctionParameterString(Function, OutReferenceGatherers);
+
+				const TCHAR* Virtual = (!FunctionData.FunctionReference->HasAnyFunctionFlags(FUNC_Static) && !(FunctionData.FunctionExportFlags & FUNCEXPORT_Final)) ? TEXT("virtual") : TEXT("");
+				FStringOutputDevice ValidDecl;
+				ValidDecl.Logf(TEXT("\t%s bool %s(%s);\r\n"), Virtual, *FunctionData.CppValidationImplName, *ParameterList);
+				FuncStringBuilders.AutogeneratedBlueprintFunctionDeclarations.Log(*ValidDecl);
+				if (!bHasValidate)
+				{
+					FuncStringBuilders.AutogeneratedBlueprintFunctionDeclarationsOnlyNotDeclared.Logf(TEXT("%s"), *ValidDecl);
+				}
+			}
+
+			FuncStringBuilders.AutogeneratedBlueprintFunctionDeclarations.Log(*FunctionDeclaration);
+			if (!bHasImplementation && FunctionData.CppImplName != FunctionName)
+			{
+				FuncStringBuilders.AutogeneratedBlueprintFunctionDeclarationsOnlyNotDeclared.Log(*FunctionDeclaration);
+			}
+
+			// Versions that skip function autodeclaration throw an error when a function is missing.
+			if (ClassRange.bHasGeneratedBody && (SourceFile.GetGeneratedCodeVersionForStruct(Class) > EGeneratedCodeVersion::V1))
+			{
+				CheckRPCFunctions(OutReferenceGatherers, FunctionData, ClassCPPName, ImplementationPosition, ValidatePosition, SourceFile);
+			}
+		}
+
+		FuncStringBuilders.RPCWrappers.Log(TEXT("\r\n"));
+
+		// if this function was originally declared in a base class, and it isn't a static function,
+		// only the C++ function header will be exported
+		if (!ShouldExportUFunction(Function))
+		{
+			continue;
+		}
+
+		// export the script wrappers
+		FuncStringBuilders.RPCWrappers.Logf(TEXT("\tDECLARE_FUNCTION(%s);"), *FunctionData.UnMarshallAndCallName);
+		FuncStringBuilders.RPCImplementations.Logf(TEXT("\tDEFINE_FUNCTION(%s::%s)"), *ClassCPPName, *FunctionData.UnMarshallAndCallName);
+		FuncStringBuilders.RPCImplementations += LINE_TERMINATOR TEXT("\t{") LINE_TERMINATOR;
+
+		FParmsAndReturnProperties Parameters = GetFunctionParmsAndReturn(FunctionData.FunctionReference);
+		ExportFunctionThunk(FuncStringBuilders.RPCImplementations, OutReferenceGatherers, Function, FunctionData, Parameters.Parms, Parameters.Return);
+
+		FuncStringBuilders.RPCImplementations += TEXT("\t}") LINE_TERMINATOR;
+	}
+
+	// static class data
+	{
+		FString MacroName = SourceFile.GetGeneratedMacroName(ClassData, TEXT("_SPARSE_DATA"));
+
+		WriteMacro(OutGeneratedHeaderText, MacroName, RuntimeStringBuilders.AutogeneratedStaticData + RuntimeStringBuilders.AutogeneratedStaticDataFuncs);
+		OutMacroCalls.Logf(TEXT("\t%s\r\n"), *MacroName);
+		OutNoPureDeclsMacroCalls.Logf(TEXT("\t%s\r\n"), *MacroName);
+	}
+
+	// Write runtime wrappers
+	{
+		FString MacroName = SourceFile.GetGeneratedMacroName(ClassData, TEXT("_RPC_WRAPPERS"));
+
+		// WriteMacro has an assumption about what will be at the end of this block that is no longer true due to splitting the
+		// definition and implementation, so add on a line terminator to satisfy it
+		if (RuntimeStringBuilders.RPCWrappers.Len() > 0)
+		{
+			RuntimeStringBuilders.RPCWrappers += LINE_TERMINATOR;
+		}
+
+		WriteMacro(OutGeneratedHeaderText, MacroName, RuntimeStringBuilders.AutogeneratedBlueprintFunctionDeclarations + RuntimeStringBuilders.RPCWrappers);
+		OutMacroCalls.Logf(TEXT("\t%s\r\n"), *MacroName);
+
+		// Put static checks before RPCWrappers to get proper messages from static asserts before compiler errors.
+		FString NoPureDeclsMacroName = SourceFile.GetGeneratedMacroName(ClassData, TEXT("_RPC_WRAPPERS_NO_PURE_DECLS"));
+		if (SourceFile.GetGeneratedCodeVersionForStruct(Class) > EGeneratedCodeVersion::V1)
+		{
+			WriteMacro(OutGeneratedHeaderText, NoPureDeclsMacroName, RuntimeStringBuilders.RPCWrappers);
+		}
+		else
+		{
+			WriteMacro(OutGeneratedHeaderText, NoPureDeclsMacroName, RuntimeStringBuilders.AutogeneratedBlueprintFunctionDeclarationsOnlyNotDeclared + RuntimeStringBuilders.RPCWrappers);
+		}
+
+		OutNoPureDeclsMacroCalls.Logf(TEXT("\t%s\r\n"), *NoPureDeclsMacroName);
+
+		OutGeneratedCPPText.Log(RuntimeStringBuilders.RPCImplementations);
+	}
+
+	// Write editor only RPC wrappers if they exist
+	if (EditorStringBuilders.RPCWrappers.Len() > 0)
+	{
+		OutGeneratedHeaderText.Log( BeginEditorOnlyGuard );
+
+		FString MacroName = SourceFile.GetGeneratedMacroName(ClassData, TEXT("_EDITOR_ONLY_RPC_WRAPPERS"));
+
+		// WriteMacro has an assumption about what will be at the end of this block that is no longer true due to splitting the
+		// definition and implementation, so add on a line terminator to satisfy it
+		if (EditorStringBuilders.RPCWrappers.Len() > 0)
+		{
+			EditorStringBuilders.RPCWrappers += LINE_TERMINATOR;
+		}
+
+		WriteMacro(OutGeneratedHeaderText, MacroName, EditorStringBuilders.AutogeneratedBlueprintFunctionDeclarations + EditorStringBuilders.RPCWrappers);
+		OutMacroCalls.Logf(TEXT("\t%s\r\n"), *MacroName);
+
+		// Put static checks before RPCWrappers to get proper messages from static asserts before compiler errors.
+		FString NoPureDeclsMacroName = SourceFile.GetGeneratedMacroName(ClassData, TEXT("_EDITOR_ONLY_RPC_WRAPPERS_NO_PURE_DECLS"));
+		if (SourceFile.GetGeneratedCodeVersionForStruct(Class) > EGeneratedCodeVersion::V1)
+		{
+			WriteMacro(OutGeneratedHeaderText, NoPureDeclsMacroName, EditorStringBuilders.RPCWrappers);
+		}
+		else
+		{
+			WriteMacro(OutGeneratedHeaderText, NoPureDeclsMacroName, EditorStringBuilders.AutogeneratedBlueprintFunctionDeclarationsOnlyNotDeclared + EditorStringBuilders.RPCWrappers);
+		}
+
+		// write out an else preprocessor block for when not compiling for the editor.  The generated macros should be empty then since the functions are compiled out
+		{
+			OutGeneratedHeaderText.Log(TEXT("#else\r\n"));
+
+			WriteMacro(OutGeneratedHeaderText, MacroName, TEXT(""));
+			WriteMacro(OutGeneratedHeaderText, NoPureDeclsMacroName, TEXT(""));
+
+			OutGeneratedHeaderText.Log(EndEditorOnlyGuard);
+		}
+
+		OutNoPureDeclsMacroCalls.Logf(TEXT("\t%s\r\n"), *NoPureDeclsMacroName);
+
+		OutGeneratedCPPText.Log(BeginEditorOnlyGuard);
+		OutGeneratedCPPText.Log(EditorStringBuilders.RPCImplementations);
+		OutGeneratedCPPText.Log(EndEditorOnlyGuard);
+	}
+}
+
+/**
+ * Exports the methods which trigger UnrealScript events and delegates.
+ *
+ * @param	CallbackFunctions	the functions to export
+ */
+void FNativeClassHeaderGenerator::ExportCallbackFunctions(
+	FOutputDevice&            OutGeneratedHeaderText,
+	FOutputDevice&            OutCpp,
+	TSet<FString>&            OutFwdDecls,
+	const TArray<UFunction*>& CallbackFunctions,
+	const TCHAR*              CallbackWrappersMacroName,
+	EExportCallbackType       ExportCallbackType,
+	const TCHAR*              APIString
+)
+{
+	FUHTStringBuilder RPCWrappers;
+
+	FMacroBlockEmitter OutCppEditorOnly(OutCpp, TEXT("WITH_EDITOR"));
+	for (UFunction* Function : CallbackFunctions)
+	{
+		// Never expecting to export delegate functions this way
+		check(!Function->HasAnyFunctionFlags(FUNC_Delegate));
+
+		FFunctionData*   CompilerInfo = FFunctionData::FindForFunction(Function);
+		const FFuncInfo& FunctionData = CompilerInfo->GetFunctionData();
+		FString          FunctionName = Function->GetName();
+		UClass*          Class = CastChecked<UClass>(Function->GetOuter());
+		const FString    ClassName = FNameLookupCPP::GetNameCPP(Class);
+
+		if (FunctionData.FunctionFlags & FUNC_NetResponse)
+		{
+			// Net response functions don't go into the VM
+			continue;
+		}
+
+		const bool bIsEditorOnly = Function->HasAnyFunctionFlags(FUNC_EditorOnly);
+
+		OutCppEditorOnly(bIsEditorOnly);
+
+		const bool bWillBeProgrammerTyped = FunctionName == FunctionData.MarshallAndCallName;
+
+		// Emit the declaration if the programmer isn't responsible for declaring this wrapper
+		if (!bWillBeProgrammerTyped)
+		{
+			// export the line that looks like: int32 Main(const FString& Parms)
+			ExportNativeFunctionHeader(RPCWrappers, OutFwdDecls, FunctionData, EExportFunctionType::Event, EExportFunctionHeaderStyle::Declaration, nullptr, APIString);
+
+			RPCWrappers.Log(TEXT(";\r\n"));
+			RPCWrappers.Log(TEXT("\r\n"));
+		}
+
+		FString FunctionNameName;
+		if (ExportCallbackType != EExportCallbackType::Interface)
+		{
+			FunctionNameName = FString::Printf(TEXT("NAME_%s_%s"), *ClassName, *FunctionName);
+			OutCpp.Logf(TEXT("\tstatic FName %s = FName(TEXT(\"%s\"));") LINE_TERMINATOR, *FunctionNameName, *GetOverriddenFName(Function).ToString());
+		}
+
+		// Emit the thunk implementation
+		ExportNativeFunctionHeader(OutCpp, OutFwdDecls, FunctionData, EExportFunctionType::Event, EExportFunctionHeaderStyle::Definition, nullptr, APIString);
+
+		FParmsAndReturnProperties Parameters = GetFunctionParmsAndReturn(FunctionData.FunctionReference);
+
+		if (ExportCallbackType != EExportCallbackType::Interface)
+		{
+			WriteEventFunctionPrologue(OutCpp, /*Indent=*/ 1, Parameters, Class, *FunctionName);
+			{
+				// Cast away const just in case, because ProcessEvent isn't const
+				OutCpp.Logf(
+					TEXT("\t\t%sProcessEvent(FindFunctionChecked(%s),%s);\r\n"),
+					(Function->HasAllFunctionFlags(FUNC_Const)) ? *FString::Printf(TEXT("const_cast<%s*>(this)->"), *ClassName) : TEXT(""),
+					*FunctionNameName,
+					Parameters.HasParms() ? TEXT("&Parms") : TEXT("NULL")
+				);
+			}
+			WriteEventFunctionEpilogue(OutCpp, /*Indent=*/ 1, Parameters);
+		}
+		else
+		{
+			OutCpp.Log(LINE_TERMINATOR);
+			OutCpp.Log(TEXT("\t{") LINE_TERMINATOR);
+
+			// assert if this is ever called directly
+			OutCpp.Logf(TEXT("\t\tcheck(0 && \"Do not directly call Event functions in Interfaces. Call Execute_%s instead.\");") LINE_TERMINATOR, *FunctionName);
+
+			// satisfy compiler if it's expecting a return value
+			if (Parameters.Return)
+			{
+				FString EventParmStructName = GetEventStructParamsName(Class, *FunctionName);
+				OutCpp.Logf(TEXT("\t\t%s Parms;") LINE_TERMINATOR, *EventParmStructName);
+				OutCpp.Log(TEXT("\t\treturn Parms.ReturnValue;") LINE_TERMINATOR);
+			}
+			OutCpp.Log(TEXT("\t}") LINE_TERMINATOR);
+		}
+	}
+
+	WriteMacro(OutGeneratedHeaderText, CallbackWrappersMacroName, RPCWrappers);
+}
+
+
+/**
+ * Determines if the property has alternate export text associated with it and if so replaces the text in PropertyText with the
+ * alternate version. (for example, structs or properties that specify a native type using export-text).  Should be called immediately
+ * after ExportCppDeclaration()
+ *
+ * @param	Prop			the property that is being exported
+ * @param	PropertyText	the string containing the text exported from ExportCppDeclaration
+ */
+void FNativeClassHeaderGenerator::ApplyAlternatePropertyExportText(FProperty* Prop, FUHTStringBuilder& PropertyText, EExportingState ExportingState)
+{
+	FArrayProperty* ArrayProperty = CastField<FArrayProperty>(Prop);
+	FProperty* InnerProperty = ArrayProperty ? ArrayProperty->Inner : nullptr;
+	if (InnerProperty && (
+			(InnerProperty->IsA<FByteProperty>() && ((FByteProperty*)InnerProperty)->Enum && FClass::IsDynamic(static_cast<UField*>(((FByteProperty*)InnerProperty)->Enum))) ||
+			(InnerProperty->IsA<FEnumProperty>()                                          && FClass::IsDynamic(static_cast<UField*>(((FEnumProperty*)InnerProperty)->Enum)))
+		)
+	)
+	{
+		const FString Original = InnerProperty->GetCPPType();
+		const FString RawByte = InnerProperty->GetCPPType(nullptr, EPropertyExportCPPFlags::CPPF_BlueprintCppBackend);
+		if (Original != RawByte)
+		{
+			PropertyText.ReplaceInline(*Original, *RawByte, ESearchCase::CaseSensitive);
+		}
+		return;
+	}
+
+	if (ExportingState == EExportingState::TypeEraseDelegates)
+	{
+		FDelegateProperty* DelegateProperty = CastField<FDelegateProperty>(Prop);
+		FMulticastDelegateProperty* MulticastDelegateProperty = CastField<FMulticastDelegateProperty>(Prop);
+		if (DelegateProperty || MulticastDelegateProperty)
+		{
+			FString Original = Prop->GetCPPType();
+			const TCHAR* PlaceholderOfSameSizeAndAlignemnt;
+			if (DelegateProperty)
+			{
+				PlaceholderOfSameSizeAndAlignemnt = TEXT("FScriptDelegate");
+			}
+			else
+			{
+				PlaceholderOfSameSizeAndAlignemnt = TEXT("FMulticastScriptDelegate");
+			}
+			PropertyText.ReplaceInline(*Original, PlaceholderOfSameSizeAndAlignemnt, ESearchCase::CaseSensitive);
+		}
+	}
+}
+
+void GetSourceFilesInDependencyOrderRecursive(TArray<FUnrealSourceFile*>& OutTest, const UPackage* Package, FUnrealSourceFile* SourceFile, TSet<const FUnrealSourceFile*>& VisitedSet, bool bCheckDependenciesOnly, const TSet<FUnrealSourceFile*>& Ignore)
+{
+	// Check if the Class has already been exported, after we've checked for circular header dependencies.
+	if (OutTest.Contains(SourceFile) || Ignore.Contains(SourceFile))
+	{
+		return;
+	}
+
+	// Check for circular dependencies.
+	if (VisitedSet.Contains(SourceFile))
+	{
+		UE_LOG(LogCompile, Error, TEXT("Circular dependency detected for filename %s!"), *SourceFile->GetFilename());
+		return;
+	}
+
+	// Check for circular header dependencies between export classes.
+	bCheckDependenciesOnly = bCheckDependenciesOnly || SourceFile->GetPackage() != Package;
+
+	VisitedSet.Add(SourceFile);
+	for (FHeaderProvider& Include : SourceFile->GetIncludes())
+	{
+		if (FUnrealSourceFile* IncludeFile = Include.Resolve())
+		{
+			GetSourceFilesInDependencyOrderRecursive(OutTest, Package, IncludeFile, VisitedSet, bCheckDependenciesOnly, Ignore);
+		}
+	}
+	VisitedSet.Remove(SourceFile);
+
+	if (!bCheckDependenciesOnly)
+	{
+		OutTest.Add(SourceFile);
+	}
+}
+
+TArray<FUnrealSourceFile*> GetSourceFilesInDependencyOrder(const UPackage* Package, const TSet<FUnrealSourceFile*>& SourceFiles, const TSet<FUnrealSourceFile*>& Ignore)
+{
+	TArray<FUnrealSourceFile*> Result;
+	TSet<const FUnrealSourceFile*>	VisitedSet;
+	for (FUnrealSourceFile* SourceFile : SourceFiles)
+	{
+		if (SourceFile->GetPackage() == Package)
+		{
+			GetSourceFilesInDependencyOrderRecursive(Result, Package, SourceFile, VisitedSet, false, Ignore);
+		}
+	}
+
+	return Result;
+}
+
+TMap<UClass*, FUnrealSourceFile*> GClassToSourceFileMap;
+
+static bool HasDynamicOuter(UField* Field)
+{
+	UField* FieldOuter = Cast<UField>(Field->GetOuter());
+	return FieldOuter && FClass::IsDynamic(FieldOuter);
+}
+
+static void RecordPackageSingletons(
+	const UPackage& Package,
+	const TArray<UScriptStruct*>& Structs,
+	const TArray<UDelegateFunction*>& Delegates)
+{
+	TArray<UField*> Singletons;
+	Singletons.Reserve(Structs.Num() + Delegates.Num());
+	for (UScriptStruct* Struct : Structs)
+	{
+		if (Struct->StructFlags & STRUCT_NoExport && !HasDynamicOuter(Struct))
+		{
+			Singletons.Add(Struct);
+		}
+	}
+
+	for (UDelegateFunction* Delegate : Delegates)
+	{
+		if (!HasDynamicOuter(Delegate))
+		{
+			Singletons.Add(Delegate);
+		}
+	}
+
+	if (Singletons.Num())
+	{
+		FScopeLock PackageSingletonLock(&GPackageSingletonsCriticalSection);
+
+		TArray<UField*>& PackageSingletons = GPackageSingletons.FindOrAdd(&Package);
+		PackageSingletons.Append(MoveTemp(Singletons));
+	}
+}
+
+struct FPreloadHeaderFileInfo
+{
+	FPreloadHeaderFileInfo()
+		: bFinishedLoading(true)
+	{
+	}
+
+	~FPreloadHeaderFileInfo()
+	{
+		EnsureLoadComplete();
+	}
+
+	void Load(FString&& InHeaderPath)
+	{
+		if (!bFinishedLoading || HeaderFileContents.Len() > 0)
+		{
+			if (ensureMsgf(InHeaderPath == HeaderPath, TEXT("FPreloadHeaderFileInfo::Load called twice with different paths.")))
+			{
+				// If we've done an async load but now a sync load has been requested we need to wait on it
+				EnsureLoadComplete();
+			}
+		}
+		else
+		{
+			HeaderPath = MoveTemp(InHeaderPath);
+
+			SCOPE_SECONDS_COUNTER_UHT(LoadHeaderContentFromFile);
+			FFileHelper::LoadFileToString(HeaderFileContents, *HeaderPath);
+		}
+	}
+
+	void StartLoad(FString&& InHeaderPath)
+	{
+		if (!bFinishedLoading || HeaderFileContents.Len() > 0)
+		{
+			ensureMsgf(InHeaderPath == HeaderPath, TEXT("FPreloadHeaderFileInfo::Load called twice with different paths."));
+			return;
+		}
+		auto LoadFileContentsTask = [this]()
+		{
+			SCOPE_SECONDS_COUNTER_UHT(LoadHeaderContentFromFile);
+			FFileHelper::LoadFileToString(HeaderFileContents, *HeaderPath);
+
+			bFinishedLoading = true;
+		};
+
+		HeaderPath = MoveTemp(InHeaderPath);
+		bFinishedLoading = false;
+
+		LoadTaskRef = FFunctionGraphTask::CreateAndDispatchWhenReady(MoveTemp(LoadFileContentsTask), TStatId());
+	}
+
+	const FString& GetFileContents() const
+	{
+		EnsureLoadComplete();
+		return HeaderFileContents;
+	}
+
+	FString& GetHeaderPath() { return HeaderPath; }
+	const FString& GetHeaderPath() const { return HeaderPath; }
+
+private:
+
+	void EnsureLoadComplete() const
+	{
+		if (!bFinishedLoading)
+		{
+			FTaskGraphInterface::Get().WaitUntilTaskCompletes(LoadTaskRef);
+		}
+	}
+
+	FString HeaderPath;
+	FString HeaderFileContents;
+	FGraphEventRef LoadTaskRef;
+	bool bFinishedLoading;
+};
+
+// Constructor.
+FNativeClassHeaderGenerator::FNativeClassHeaderGenerator(
+	const UPackage* InPackage,
+	const TSet<FUnrealSourceFile*>& SourceFiles,
+	FClasses& AllClasses,
+	bool InAllowSaveExportedHeaders
+)
+	: API                        (FPackageName::GetShortName(InPackage).ToUpper())
+	, APIStringPrivate           (FString::Printf(TEXT("%s_API "), *API))
+	, Package                    (InPackage)
+	, bAllowSaveExportedHeaders  (InAllowSaveExportedHeaders)
+{
+	FString PackageName = FPackageName::GetShortName(Package);
+
+	FManifestModule* PackageManifest = GetPackageManifest(PackageName);
+	if (!PackageManifest)
+	{
+		UE_LOG(LogCompile, Error, TEXT("Failed to find path for package %s"), *PackageName);
+	}
+
+	bool bWriteClassesH = false;
+	const bool bPackageHasAnyExportClasses = AllClasses.GetClassesInPackage(Package).ContainsByPredicate([](FClass* Class)
+	{
+		return Class->HasAnyClassFlags(CLASS_Native) && !Class->HasAnyClassFlags(CLASS_NoExport | CLASS_Intrinsic);
+	});
+	if (bPackageHasAnyExportClasses)
+	{
+		for (FUnrealSourceFile* SourceFile : SourceFiles)
+		{
+			for (const TPair<UClass*, FSimplifiedParsingClassInfo>& ClassDataPair : SourceFile->GetDefinedClassesWithParsingInfo())
+			{
+				UClass* Class = ClassDataPair.Key;
+				if (!Class->HasAnyClassFlags(CLASS_Native))
+				{
+					Class->UnMark(EObjectMark(OBJECTMARK_TagImp | OBJECTMARK_TagExp));
+				}
+				else if (!Class->HasAnyClassFlags(CLASS_NoExport) && GTypeDefinitionInfoMap.Contains(Class))
+				{
+					bWriteClassesH = true;
+					Class->UnMark(OBJECTMARK_TagImp);
+					Class->Mark(OBJECTMARK_TagExp);
+				}
+			}
+		}
+	}
+
+	TArray<FUnrealSourceFile*> Exported;
+	{
+		// Get source files and ignore them next time round
+		static TSet<FUnrealSourceFile*> ExportedSourceFiles;
+		Exported = GetSourceFilesInDependencyOrder(Package, SourceFiles, ExportedSourceFiles);
+		ExportedSourceFiles.Append(Exported);
+	}
+
+	struct FGeneratedCPP
+	{
+		explicit FGeneratedCPP(FString&& InGeneratedCppFullFilename)
+			: GeneratedCppFullFilename(MoveTemp(InGeneratedCppFullFilename))
+		{
+		}
+
+		FString                      GeneratedCppFullFilename;
+		TArray<FString>              RelativeIncludes;
+		FUHTStringBuilderLineCounter GeneratedText;
+		TSet<FString>                CrossModuleReferences;
+		TSet<FString>                PackageHeaderPaths;
+		TArray<FString>              TempHeaderPaths;
+		FUHTStringBuilder			 GeneratedFunctionDeclarations;
+	};
+
+	TMap<FUnrealSourceFile*, FGeneratedCPP> GeneratedCPPs;
+	GeneratedCPPs.Reserve(Exported.Num());
+
+	// Set up the generated cpp map
+	for (FUnrealSourceFile* SourceFile : Exported)
+	{
+		FString ModuleRelativeFilename = SourceFile->GetFilename();
+		ConvertToBuildIncludePath(Package, ModuleRelativeFilename);
+
+		FString StrippedName = FPaths::GetBaseFilename(ModuleRelativeFilename);
+		FString GeneratedSourceFilename = (PackageManifest->GeneratedIncludeDirectory / StrippedName) + TEXT(".gen.cpp");
+
+		FGeneratedCPP& GeneratedCPP = GeneratedCPPs.Emplace(SourceFile, MoveTemp(GeneratedSourceFilename));
+		GeneratedCPP.RelativeIncludes.Add(MoveTemp(ModuleRelativeFilename));
+
+		// This needs to be done outside of parallel blocks because it will modify UClass memory.
+		// Later calls to SetUpUhtReplicationData inside parallel blocks should be fine, because
+		// they will see the memory has already been set up, and just return the parent pointer.
+		for (const TPair<UClass*, FSimplifiedParsingClassInfo>& ClassDataPair : SourceFile->GetDefinedClassesWithParsingInfo())
+		{
+			UClass* Class = ClassDataPair.Key;
+			if (ClassHasReplicatedProperties(Class))
+			{
+				Class->SetUpUhtReplicationData();
+			}
+		}
+	}
+
+	const FManifestModule* ConstPackageManifest = GetPackageManifest(PackageName);
+	const FNativeClassHeaderGenerator* ConstThis = this;
+
+	TempSaveTasks.SetNum(Exported.Num());
+
+	TArray<FPreloadHeaderFileInfo> PreloadedFiles;
+	PreloadedFiles.SetNum(Exported.Num());
+
+	ParallelFor(Exported.Num(), [&Exported, &PreloadedFiles, Package=Package, ConstPackageManifest](int32 Index)
+	{
+		FUnrealSourceFile* SourceFile = Exported[Index];
+
+		FString ModuleRelativeFilename = SourceFile->GetFilename();
+		ConvertToBuildIncludePath(Package, ModuleRelativeFilename);
+
+		FString StrippedName = FPaths::GetBaseFilename(MoveTemp(ModuleRelativeFilename));
+		FString HeaderPath = (ConstPackageManifest->GeneratedIncludeDirectory / StrippedName) + TEXT(".generated.h");
+
+		PreloadedFiles[Index].Load(MoveTemp(HeaderPath));
+	});
+
+	FString ExceptionMessage;
+
+	ParallelFor(Exported.Num(), [&Exported, &PreloadedFiles, &GeneratedCPPs, &TempSaveTasks=TempSaveTasks, &ExceptionMessage, ConstThis](int32 Index)
+	{
+#if !PLATFORM_EXCEPTIONS_DISABLED
+		try
+		{
+#endif		
+			FUnrealSourceFile* SourceFile = Exported[Index];
+
+			/** Forward declarations that we need for this sourcefile. */
+			TSet<FString> ForwardDeclarations;
+
+			FUHTStringBuilder GeneratedHeaderText;
+			FGeneratedCPP& GeneratedCPP = GeneratedCPPs.FindChecked(SourceFile);
+			FUHTStringBuilder& GeneratedFunctionDeclarations = GeneratedCPP.GeneratedFunctionDeclarations;
+
+			FReferenceGatherers ReferenceGatherers(&GeneratedCPP.CrossModuleReferences, GeneratedCPP.PackageHeaderPaths, GeneratedCPP.TempHeaderPaths);
+
+			FUHTStringBuilderLineCounter& OutText = GeneratedCPP.GeneratedText;
+
+			TArray<UEnum*> Enums;
+			TArray<UScriptStruct*> Structs;
+			TArray<UDelegateFunction*> DelegateFunctions;
+			SourceFile->GetScope()->SplitTypesIntoArrays(Enums, Structs, DelegateFunctions);
+
+			RecordPackageSingletons(*SourceFile->GetPackage(), Structs, DelegateFunctions);
+
+			// Reverse the containers as they come out in the reverse order of declaration
+			Algo::Reverse(Enums);
+			Algo::Reverse(Structs);
+			Algo::Reverse(DelegateFunctions);
+
+			const FString FileDefineName = SourceFile->GetFileDefineName();
+			const FString& StrippedFilename = SourceFile->GetStrippedFilename();
+
+			GeneratedHeaderText.Logf(
+				TEXT("#ifdef %s")																	LINE_TERMINATOR
+				TEXT("#error \"%s.generated.h already included, missing '#pragma once' in %s.h\"")	LINE_TERMINATOR
+				TEXT("#endif")																		LINE_TERMINATOR
+				TEXT("#define %s")																	LINE_TERMINATOR
+				LINE_TERMINATOR,
+				*FileDefineName, *StrippedFilename, *StrippedFilename, *FileDefineName);
+
+			// export delegate definitions
+			for (UDelegateFunction* Func : DelegateFunctions)
+			{
+				GeneratedFunctionDeclarations.Log(FTypeSingletonCache::Get(Func).GetExternDecl());
+				ConstThis->ExportDelegateDeclaration(OutText, ReferenceGatherers, *SourceFile, Func);
+			}
+
+			// Export enums declared in non-UClass headers.
+			for (UEnum* Enum : Enums)
+			{
+				// Is this ever not the case?
+				if (Enum->GetOuter()->IsA(UPackage::StaticClass()))
+				{
+					GeneratedFunctionDeclarations.Log(FTypeSingletonCache::Get(Enum).GetExternDecl());
+					ConstThis->ExportGeneratedEnumInitCode(OutText, ReferenceGatherers, *SourceFile, Enum);
+				}
+			}
+
+			// export boilerplate macros for structs
+			// reverse the order.
+			for (UScriptStruct* Struct : Structs)
+			{
+				GeneratedFunctionDeclarations.Log(FTypeSingletonCache::Get(Struct).GetExternDecl());
+				ConstThis->ExportGeneratedStructBodyMacros(GeneratedHeaderText, OutText, ReferenceGatherers, *SourceFile, Struct);
+			}
+
+			// export delegate wrapper function implementations
+			for (UDelegateFunction* Func : DelegateFunctions)
+			{
+				ConstThis->ExportDelegateDefinition(GeneratedHeaderText, ReferenceGatherers, *SourceFile, Func);
+			}
+
+			EExportClassOutFlags ExportFlags = EExportClassOutFlags::None;
+			TSet<FString> AdditionalHeaders;
+			for (const TPair<UClass*, FSimplifiedParsingClassInfo>& ClassDataPair : SourceFile->GetDefinedClassesWithParsingInfo())
+			{
+				UClass* Class = ClassDataPair.Key;
+				if (!(Class->ClassFlags & CLASS_Intrinsic))
+				{
+					ConstThis->ExportClassFromSourceFileInner(GeneratedHeaderText, OutText, GeneratedFunctionDeclarations, ReferenceGatherers, (FClass*)Class, *SourceFile, ExportFlags);
+				}
+			}
+
+			if (EnumHasAnyFlags(ExportFlags, EExportClassOutFlags::NeedsPushModelHeaders))
+			{
+				AdditionalHeaders.Add(FString(TEXT("Net/Core/PushModel/PushModelMacros.h")));
+			}
+
+			GeneratedHeaderText.Log(TEXT("#undef CURRENT_FILE_ID\r\n"));
+			GeneratedHeaderText.Logf(TEXT("#define CURRENT_FILE_ID %s\r\n\r\n\r\n"), *SourceFile->GetFileId());
+
+			for (UEnum* Enum : Enums)
+			{
+				ConstThis->ExportEnum(GeneratedHeaderText, Enum);
+			}
+
+			FPreloadHeaderFileInfo& FileInfo = PreloadedFiles[Index];
+			bool bHasChanged = ConstThis->WriteHeader(FileInfo, GeneratedHeaderText, AdditionalHeaders, ReferenceGatherers, TempSaveTasks[Index]);
+
+			SourceFile->SetGeneratedFilename(MoveTemp(FileInfo.GetHeaderPath()));
+			SourceFile->SetHasChanged(bHasChanged);
+#if !PLATFORM_EXCEPTIONS_DISABLED
+		}
+		catch (const TCHAR* Ex)
+		{
+			// Capture the first exception message from the loop and issue it out on the gamethread after the loop completes
+
+			static FCriticalSection ExceptionCS;
+			FScopeLock Lock(&ExceptionCS);
+			
+			if (ExceptionMessage.Len() == 0)
+			{
+				ExceptionMessage = FString(Ex);
+			}
+		}
+#endif
+	});
+
+	if (ExceptionMessage.Len() > 0)
+	{
+		FError::Throwf(*ExceptionMessage);
+	}
+
+	FPreloadHeaderFileInfo FileInfo;
+	if (bWriteClassesH)
+	{
+		// Start loading the original header file for comparison
+		FString ClassesHeaderPath = PackageManifest->GeneratedIncludeDirectory / (PackageName + TEXT("Classes.h"));
+		FileInfo.StartLoad(MoveTemp(ClassesHeaderPath));
+	}
+
+	// Export an include line for each header
+	TSet<FUnrealSourceFile*> PublicHeaderGroupIncludes;
+	FUHTStringBuilder GeneratedFunctionDeclarations;
+
+	for (FUnrealSourceFile* SourceFile : Exported)
+	{
+		for (const TPair<UClass*, FSimplifiedParsingClassInfo>& ClassDataPair : SourceFile->GetDefinedClassesWithParsingInfo())
+		{
+			UClass* Class = ClassDataPair.Key;
+			GClassToSourceFileMap.Add(Class, SourceFile);
+		}
+
+		if (GPublicSourceFileSet.Contains(SourceFile))
+		{
+			PublicHeaderGroupIncludes.Add(SourceFile);
+		}
+
+		const FGeneratedCPP& GeneratedCPP = GeneratedCPPs.FindChecked(SourceFile);
+		GeneratedFunctionDeclarations.Log(GeneratedCPP.GeneratedFunctionDeclarations);
+	}
+
+	// Add includes for 'Within' classes
+	for (FUnrealSourceFile* SourceFile : Exported)
+	{
+		bool bAddedStructuredArchiveFromArchiveHeader = false;
+		bool bAddedArchiveUObjectFromStructuredArchiveHeader = false;
+
+		TArray<FString>& RelativeIncludes = GeneratedCPPs[SourceFile].RelativeIncludes;
+		for (const TPair<UClass*, FSimplifiedParsingClassInfo>& ClassDataPair : SourceFile->GetDefinedClassesWithParsingInfo())
+		{
+			UClass* Class = ClassDataPair.Key;
+			if (Class->ClassWithin && Class->ClassWithin != UObject::StaticClass())
+			{
+				if (FUnrealSourceFile** WithinSourceFile = GClassToSourceFileMap.Find(Class->ClassWithin))
+				{
+					FString Header = GetBuildPath(**WithinSourceFile);
+					RelativeIncludes.AddUnique(MoveTemp(Header));
+				}
+			}
+
+			if (const FArchiveTypeDefinePair* ArchiveTypeDefinePair = GClassSerializerMap.Find(Class))
+			{
+				if (!bAddedStructuredArchiveFromArchiveHeader && ArchiveTypeDefinePair->ArchiveType == ESerializerArchiveType::StructuredArchiveRecord)
+				{
+					RelativeIncludes.AddUnique(TEXT("Serialization/StructuredArchive.h"));
+					bAddedStructuredArchiveFromArchiveHeader = true;
+				}
+
+				if (!bAddedArchiveUObjectFromStructuredArchiveHeader && ArchiveTypeDefinePair->ArchiveType == ESerializerArchiveType::Archive)
+				{
+					RelativeIncludes.AddUnique(TEXT("Serialization/ArchiveUObjectFromStructuredArchive.h"));
+					bAddedArchiveUObjectFromStructuredArchiveHeader = true;
+				}
+			}
+		}
+	}
+
+	TSet<FString> PackageHeaderPaths;
+	TArray<FString> TempHeaderPaths;
+	if (bWriteClassesH)
+	{
+		// Write the classes and enums header prefixes.
+
+		FUHTStringBuilder ClassesHText;
+		ClassesHText.Log(HeaderCopyright);
+		ClassesHText.Log(TEXT("#pragma once\r\n"));
+		ClassesHText.Log(TEXT("\r\n"));
+		ClassesHText.Log(TEXT("\r\n"));
+
+		// Fill with the rest source files from this package.
+		if (const TArray<FUnrealSourceFile*>* SourceFilesForPackage = GPublicSourceFileSet.FindFilesForPackage(InPackage))
+		{
+			PublicHeaderGroupIncludes.Append(*SourceFilesForPackage);
+		}
+
+		for (FUnrealSourceFile* SourceFile : PublicHeaderGroupIncludes)
+		{
+			ClassesHText.Logf(TEXT("#include \"%s\"") LINE_TERMINATOR, *GetBuildPath(*SourceFile));
+		}
+
+		ClassesHText.Log(LINE_TERMINATOR);
+
+		FReferenceGatherers ReferenceGatherers(nullptr, PackageHeaderPaths, TempHeaderPaths);
+
+		// Save the classes header if it has changed.
+		FGraphEventRef& SaveTaskRef = TempSaveTasks.AddDefaulted_GetRef();
+		SaveHeaderIfChanged(ReferenceGatherers, FileInfo, MoveTemp(ClassesHText), SaveTaskRef);
+	}
+
+	// now export the names for the functions in this package
+	// notice we always export this file (as opposed to only exporting if we have any marked names)
+	// because there would be no way to know when the file was created otherwise
+	UE_LOG(LogCompile, Log, TEXT("Generating code for module '%s'"), *PackageName);
+
+	if (GeneratedFunctionDeclarations.Len())
+	{
+		uint32 CombinedHash = 0;
+		for (const TPair<FUnrealSourceFile*, FGeneratedCPP>& GeneratedCPP : GeneratedCPPs)
+		{
+			uint32 SplitHash = GenerateTextHash(*GeneratedCPP.Value.GeneratedText);
+			if (CombinedHash == 0)
+			{
+				// Don't combine in the first case because it keeps GUID backwards compatibility
+				CombinedHash = SplitHash;
+			}
+			else
+			{
+				CombinedHash = HashCombine(SplitHash, CombinedHash);
+			}
+		}
+
+		FGeneratedCPP& GeneratedCPP = GeneratedCPPs.Emplace(nullptr, PackageManifest->GeneratedIncludeDirectory / FString::Printf(TEXT("%s.init.gen.cpp"), *PackageName));
+		ExportGeneratedPackageInitCode(GeneratedCPP.GeneratedText, *GeneratedFunctionDeclarations, Package, CombinedHash);
+	}
+
+	const FManifestModule* ModuleInfo = GPackageToManifestModuleMap.FindChecked(Package);
+
+	struct FGeneratedCPPInfo
+	{
+		FGeneratedCPP* GeneratedCPP;
+		FPreloadHeaderFileInfo FileInfo;
+	};
+
+	TArray<FGeneratedCPPInfo> GeneratedCPPArray;
+	GeneratedCPPArray.SetNum(GeneratedCPPs.Num());
+
+	int32 Index = 0;
+	for (TPair<FUnrealSourceFile*, FGeneratedCPP>& Pair : GeneratedCPPs)
+	{
+		GeneratedCPPArray[Index++].GeneratedCPP = &Pair.Value;
+	}
+
+	if (bAllowSaveExportedHeaders)
+	{
+		ParallelFor(GeneratedCPPArray.Num(), [&](int32 Index)
+		{
+			FGeneratedCPPInfo& CPPInfo = GeneratedCPPArray[Index];
+			CPPInfo.FileInfo.Load(FString(CPPInfo.GeneratedCPP->GeneratedCppFullFilename));
+		});
+	}
+
+	const int32 SaveTaskStartIndex = TempSaveTasks.Num();
+	TempSaveTasks.AddDefaulted(GeneratedCPPArray.Num());
+
+	// Generate CPP files
+	ParallelFor(GeneratedCPPArray.Num(), [ConstThis, &GeneratedCPPArray, &TempSaveTasks=TempSaveTasks, SaveTaskStartIndex](int32 Index)
+	{
+		FGeneratedCPPInfo& CPPInfo = GeneratedCPPArray[Index];
+		FGeneratedCPP* GeneratedCPP = CPPInfo.GeneratedCPP;
+		FReferenceGatherers ReferenceGatherers(nullptr, GeneratedCPP->PackageHeaderPaths, GeneratedCPP->TempHeaderPaths);
+
+		FUHTStringBuilder FileText;
+
+		FString GeneratedIncludes;
+		for (const FString& RelativeInclude : GeneratedCPP->RelativeIncludes)
+		{
+			GeneratedIncludes += FString::Printf(TEXT("#include \"%s\"\r\n"), *RelativeInclude);
+		}
+
+		FString CleanFilename = FPaths::GetCleanFilename(GeneratedCPP->GeneratedCppFullFilename);
+
+		CleanFilename.ReplaceInline(TEXT(".gen.cpp"), TEXT(""), ESearchCase::CaseSensitive);
+		CleanFilename.ReplaceInline(TEXT("."), TEXT("_"), ESearchCase::CaseSensitive);
+
+		ExportGeneratedCPP(
+			FileText,
+			GeneratedCPP->CrossModuleReferences,
+			*CleanFilename,
+			*GeneratedCPP->GeneratedText,
+			*GeneratedIncludes
+		);
+
+		ConstThis->SaveHeaderIfChanged(ReferenceGatherers, CPPInfo.FileInfo, MoveTemp(FileText), TempSaveTasks[SaveTaskStartIndex+Index]);
+	});
+
+	if (bAllowSaveExportedHeaders)
+	{
+		TArray<FString> GeneratedCPPNames;
+		GeneratedCPPNames.Reserve(GeneratedCPPs.Num());
+		for (const TPair<FUnrealSourceFile*, FGeneratedCPP>& GeneratedCPP : GeneratedCPPs)
+		{
+			GeneratedCPPNames.Add(FPaths::GetCleanFilename(GeneratedCPP.Value.GeneratedCppFullFilename));
+		}
+
+		// Delete old generated .cpp files which we don't need because we generated less code than last time.
+		TArray<FString> FoundFiles;
+		FString BaseDir = FPaths::GetPath(ModuleInfo->GeneratedCPPFilenameBase);
+		IFileManager::Get().FindFiles(FoundFiles, *FPaths::Combine(BaseDir, TEXT("*.generated.cpp")), true, false);
+		IFileManager::Get().FindFiles(FoundFiles, *FPaths::Combine(BaseDir, TEXT("*.generated.*.cpp")), true, false);
+		IFileManager::Get().FindFiles(FoundFiles, *FPaths::Combine(BaseDir, TEXT("*.gen.cpp")), true, false);
+		IFileManager::Get().FindFiles(FoundFiles, *FPaths::Combine(BaseDir, TEXT("*.gen.*.cpp")), true, false);
+		for (FString& File : FoundFiles)
+		{
+			if (!GeneratedCPPNames.Contains(File))
+			{
+				IFileManager::Get().Delete(*FPaths::Combine(*BaseDir, *File));
+			}
+		}
+	}
+
+	for (TPair<FUnrealSourceFile*, FGeneratedCPP>& GeneratedCPP : GeneratedCPPs)
+	{
+		TempHeaderPaths.Append(MoveTemp(GeneratedCPP.Value.TempHeaderPaths));
+		PackageHeaderPaths.Append(MoveTemp(GeneratedCPP.Value.PackageHeaderPaths));
+	}
+
+	// Export all changed headers from their temp files to the .h files
+	ExportUpdatedHeaders(MoveTemp(PackageName), MoveTemp(TempHeaderPaths), TempSaveTasks);
+
+	// Delete stale *.generated.h files
+	DeleteUnusedGeneratedHeaders(MoveTemp(PackageHeaderPaths));
+}
+
+void FNativeClassHeaderGenerator::DeleteUnusedGeneratedHeaders(TSet<FString>&& PackageHeaderPathSet)
+{
+	auto DeleteUnusedGeneratedHeadersTask = [PackageHeaderPathSet = MoveTemp(PackageHeaderPathSet)]()
+	{
+		TSet<FString> AllIntermediateFolders;
+
+		for (const FString& PackageHeader : PackageHeaderPathSet)
+		{
+			FString IntermediatePath = FPaths::GetPath(PackageHeader);
+
+			if (AllIntermediateFolders.Contains(IntermediatePath))
+			{
+				continue;
+			}
+
+			TArray<FString> AllHeaders;
+			IFileManager::Get().FindFiles(AllHeaders, *(IntermediatePath / TEXT("*.generated.h")), true, false);
+
+			for (const FString& Header : AllHeaders)
+			{
+				const FString HeaderPath = IntermediatePath / Header;
+
+				if (PackageHeaderPathSet.Contains(HeaderPath))
+				{
+					continue;
+				}
+
+				// Check intrinsic classes. Get the class name from file name by removing .generated.h.
+				FString HeaderFilename = FPaths::GetBaseFilename(HeaderPath);
+				const int32   GeneratedIndex = HeaderFilename.Find(TEXT(".generated"), ESearchCase::IgnoreCase, ESearchDir::FromEnd);
+				const FString ClassName = MoveTemp(HeaderFilename).Mid(0, GeneratedIndex);
+				UClass* IntrinsicClass = FindObject<UClass>(ANY_PACKAGE, *ClassName);
+				if (!IntrinsicClass || !IntrinsicClass->HasAnyClassFlags(CLASS_Intrinsic))
+				{
+					IFileManager::Get().Delete(*HeaderPath);
+				}
+			}
+
+			AllIntermediateFolders.Add(MoveTemp(IntermediatePath));
+		}
+	};
+
+	GAsyncFileTasks.Add(FFunctionGraphTask::CreateAndDispatchWhenReady(MoveTemp(DeleteUnusedGeneratedHeadersTask), TStatId()));
+}
+
+/**
+ * Dirty hack global variable to allow different result codes passed through
+ * exceptions. Needs to be fixed in future versions of UHT.
+ */
+ECompilationResult::Type GCompilationResult = ECompilationResult::OtherCompilationError;
+
+bool FNativeClassHeaderGenerator::SaveHeaderIfChanged(FReferenceGatherers& OutReferenceGatherers, const FPreloadHeaderFileInfo& FileInfo, FString&& InNewHeaderContents, FGraphEventRef& OutSaveTaskRef) const
+{
+	if ( !bAllowSaveExportedHeaders )
+	{
+		// Return false indicating that the header did not need updating
+		return false;
+	}
+
+	static bool bTestedCmdLine = false;
+	if (!bTestedCmdLine)
+	{
+		bTestedCmdLine = true;
+
+		const FString& ProjectSavedDir = FPaths::ProjectSavedDir();
+
+		if (FParse::Param(FCommandLine::Get(), TEXT("WRITEREF")))
+		{
+			const FString ReferenceGeneratedCodePath = ProjectSavedDir / TEXT("ReferenceGeneratedCode/");
+
+			bWriteContents = true;
+			UE_LOG(LogCompile, Log, TEXT("********************************* Writing reference generated code to %s."), *ReferenceGeneratedCodePath);
+			UE_LOG(LogCompile, Log, TEXT("********************************* Deleting all files in ReferenceGeneratedCode."));
+			IFileManager::Get().DeleteDirectory(*ReferenceGeneratedCodePath, false, true);
+			IFileManager::Get().MakeDirectory(*ReferenceGeneratedCodePath);
+		}
+		else if (FParse::Param( FCommandLine::Get(), TEXT("VERIFYREF")))
+		{
+			const FString ReferenceGeneratedCodePath = ProjectSavedDir / TEXT("ReferenceGeneratedCode/");
+			const FString VerifyGeneratedCodePath = ProjectSavedDir / TEXT("VerifyGeneratedCode/");
+
+			bVerifyContents = true;
+			UE_LOG(LogCompile, Log, TEXT("********************************* Writing generated code to %s and comparing to %s"), *VerifyGeneratedCodePath, *ReferenceGeneratedCodePath);
+			UE_LOG(LogCompile, Log, TEXT("********************************* Deleting all files in VerifyGeneratedCode."));
+			IFileManager::Get().DeleteDirectory(*VerifyGeneratedCodePath, false, true);
+			IFileManager::Get().MakeDirectory(*VerifyGeneratedCodePath);
+		}
+	}
+
+	if (bWriteContents || bVerifyContents)
+	{
+		const FString& ProjectSavedDir = FPaths::ProjectSavedDir();
+		const FString CleanFilename = FPaths::GetCleanFilename(FileInfo.GetHeaderPath());
+		const FString Ref = ProjectSavedDir / TEXT("ReferenceGeneratedCode") / CleanFilename;
+
+		if (bWriteContents)
+		{
+			int32 i;
+			for (i = 0 ;i < 10; i++)
+			{
+				if (FFileHelper::SaveStringToFile(InNewHeaderContents, *Ref))
+				{
+					break;
+				}
+				FPlatformProcess::Sleep(1.0f); // I don't know why this fails after we delete the directory
+			}
+			check(i<10);
+		}
+		else
+		{
+			const FString Verify = ProjectSavedDir / TEXT("VerifyGeneratedCode") / CleanFilename;
+
+			int32 i;
+			for (i = 0 ;i < 10; i++)
+			{
+				if (FFileHelper::SaveStringToFile(InNewHeaderContents, *Verify))
+				{
+					break;
+				}
+				FPlatformProcess::Sleep(1.0f); // I don't know why this fails after we delete the directory
+			}
+			check(i<10);
+			FString RefHeader;
+			FString Message;
+			{
+				SCOPE_SECONDS_COUNTER_UHT(LoadHeaderContentFromFile);
+				if (!FFileHelper::LoadFileToString(RefHeader, *Ref))
+				{
+					Message = FString::Printf(TEXT("********************************* %s appears to be a new generated file."), *CleanFilename);
+				}
+				else
+				{
+					if (FCString::Strcmp(*InNewHeaderContents, *RefHeader) != 0)
+					{
+						Message = FString::Printf(TEXT("********************************* %s has changed."), *CleanFilename);
+					}
+				}
+			}
+			if (Message.Len())
+			{
+				UE_LOG(LogCompile, Log, TEXT("%s"), *Message);
+				ChangeMessages.AddUnique(MoveTemp(Message));
+			}
+		}
+	}
+
+	FString HeaderPathStr = FileInfo.GetHeaderPath();
+	const FString& OriginalHeaderLocal = FileInfo.GetFileContents();
+
+	const bool bHasChanged = OriginalHeaderLocal.Len() == 0 || FCString::Strcmp(*OriginalHeaderLocal, *InNewHeaderContents);
+	if (bHasChanged)
+	{
+		static const bool bFailIfGeneratedCodeChanges = FParse::Param(FCommandLine::Get(), TEXT("FailIfGeneratedCodeChanges"));
+		if (bFailIfGeneratedCodeChanges)
+		{
+			FString ConflictPath = HeaderPathStr + TEXT(".conflict");
+			FFileHelper::SaveStringToFile(InNewHeaderContents, *ConflictPath);
+
+			GCompilationResult = ECompilationResult::FailedDueToHeaderChange;
+			FError::Throwf(TEXT("ERROR: '%s': Changes to generated code are not allowed - conflicts written to '%s'"), *HeaderPathStr, *ConflictPath);
+		}
+
+		// save the updated version to a tmp file so that the user can see what will be changing
+		FString TmpHeaderFilename = GenerateTempHeaderName(HeaderPathStr, false);
+
+		auto SaveTempTask = [TmpHeaderFilename, InNewHeaderContents = MoveTemp(InNewHeaderContents)]()
+		{
+			// delete any existing temp file
+			IFileManager::Get().Delete(*TmpHeaderFilename, false, true);
+			if (!FFileHelper::SaveStringToFile(InNewHeaderContents, *TmpHeaderFilename))
+			{
+				UE_LOG_WARNING_UHT(TEXT("Failed to save header export preview: '%s'"), *TmpHeaderFilename);
+			}
+		};
+
+		OutSaveTaskRef = FFunctionGraphTask::CreateAndDispatchWhenReady(MoveTemp(SaveTempTask), TStatId());
+
+		OutReferenceGatherers.TempHeaderPaths.Add(MoveTemp(TmpHeaderFilename));
+	}
+
+	// Remember this header filename to be able to check for any old (unused) headers later.
+	HeaderPathStr.ReplaceInline(TEXT("\\"), TEXT("/"), ESearchCase::CaseSensitive);
+
+	OutReferenceGatherers.PackageHeaderPaths.Add(MoveTemp(HeaderPathStr));
+
+	return bHasChanged;
+}
+
+FString FNativeClassHeaderGenerator::GenerateTempHeaderName( const FString& CurrentFilename, bool bReverseOperation )
+{
+	return bReverseOperation
+		? CurrentFilename.Replace(TEXT(".tmp"), TEXT(""), ESearchCase::CaseSensitive)
+		: CurrentFilename + TEXT(".tmp");
+}
+
+void FNativeClassHeaderGenerator::ExportUpdatedHeaders(FString&& PackageName, TArray<FString>&& TempHeaderPaths, FGraphEventArray& InTempSaveTasks)
+{	
+	// Asynchronously move the headers to the correct locations
+	if (TempHeaderPaths.Num() > 0)
+	{
+		auto MoveHeadersTask = [PackageName = MoveTemp(PackageName), TempHeaderPaths = MoveTemp(TempHeaderPaths)]()
+		{
+			ParallelFor(TempHeaderPaths.Num(), [&](int32 Index)
+			{
+				const FString& TmpFilename = TempHeaderPaths[Index];
+				FString Filename = GenerateTempHeaderName(TmpFilename, true);
+				if (!IFileManager::Get().Move(*Filename, *TmpFilename, true, true))
+				{
+					UE_LOG(LogCompile, Error, TEXT("Error exporting %s: couldn't write file '%s'"), *PackageName, *Filename);
+				}
+				else
+				{
+					UE_LOG(LogCompile, Log, TEXT("Exported updated C++ header: %s"), *Filename);
+				}
+			});
+		};
+
+		FTaskGraphInterface::Get().WaitUntilTasksComplete(InTempSaveTasks);
+		GAsyncFileTasks.Add(FFunctionGraphTask::CreateAndDispatchWhenReady(MoveTemp(MoveHeadersTask), TStatId()));
+	}
+}
+
+/**
+ * Exports C++ definitions for boilerplate that was generated for a package.
+ */
+void FNativeClassHeaderGenerator::ExportGeneratedCPP(FOutputDevice& Out, const TSet<FString>& InCrossModuleReferences, const TCHAR* EmptyLinkFunctionPostfix, const TCHAR* Body, const TCHAR* OtherIncludes)
+{
+	static const TCHAR DisableWarning4883        [] = TEXT("#ifdef _MSC_VER") LINE_TERMINATOR TEXT("#pragma warning (push)") LINE_TERMINATOR TEXT("#pragma warning (disable : 4883)") LINE_TERMINATOR TEXT("#endif") LINE_TERMINATOR;
+	static const TCHAR EnableWarning4883         [] = TEXT("#ifdef _MSC_VER") LINE_TERMINATOR TEXT("#pragma warning (pop)") LINE_TERMINATOR TEXT("#endif") LINE_TERMINATOR;
+
+	Out.Log(HeaderCopyright);
+	Out.Log(RequiredCPPIncludes);
+	Out.Log(OtherIncludes);
+	Out.Log(DisableWarning4883);
+	Out.Log(DisableDeprecationWarnings);
+
+	Out.Logf(TEXT("void EmptyLinkFunctionForGeneratedCode%s() {}") LINE_TERMINATOR, EmptyLinkFunctionPostfix);
+
+	if (InCrossModuleReferences.Num() > 0)
+	{
+		Out.Logf(TEXT("// Cross Module References\r\n"));
+		for (const FString& Ref : InCrossModuleReferences)
+		{
+			Out.Log(*Ref);
+		}
+		Out.Logf(TEXT("// End Cross Module References\r\n"));
+	}
+	Out.Log(Body);
+	Out.Log(EnableDeprecationWarnings);
+	Out.Log(EnableWarning4883);
+}
+
+/** Get all script plugins based on ini setting */
+void GetScriptPlugins(TArray<IScriptGeneratorPluginInterface*>& ScriptPlugins)
+{
+	FScopedDurationTimer PluginTimeTracker(GPluginOverheadTime);
+
+	ScriptPlugins = IModularFeatures::Get().GetModularFeatureImplementations<IScriptGeneratorPluginInterface>(TEXT("ScriptGenerator"));
+	UE_LOG(LogCompile, Log, TEXT("Found %d script generator plugins."), ScriptPlugins.Num());
+
+	// Check if we can use these plugins and initialize them
+	for (int32 PluginIndex = ScriptPlugins.Num() - 1; PluginIndex >= 0; --PluginIndex)
+	{
+		IScriptGeneratorPluginInterface* ScriptGenerator = ScriptPlugins[PluginIndex];
+		bool bSupportedPlugin = ScriptGenerator->SupportsTarget(GManifest.TargetName);
+		if (bSupportedPlugin)
+		{
+			// Find the right output directory for this plugin base on its target (Engine-side) plugin name.
+			FString GeneratedCodeModuleName = ScriptGenerator->GetGeneratedCodeModuleName();
+			const FManifestModule* GeneratedCodeModule = NULL;
+			FString OutputDirectory;
+			FString IncludeBase;
+			for (const FManifestModule& Module : GManifest.Modules)
+			{
+				if (Module.Name == GeneratedCodeModuleName)
+				{
+					GeneratedCodeModule = &Module;
+				}
+			}
+			if (GeneratedCodeModule)
+			{
+				UE_LOG(LogCompile, Log, TEXT("Initializing script generator \'%s\'"), *ScriptGenerator->GetGeneratorName());
+				ScriptGenerator->Initialize(GManifest.RootLocalPath, GManifest.RootBuildPath, GeneratedCodeModule->GeneratedIncludeDirectory, GeneratedCodeModule->IncludeBase);
+			}
+			else
+			{
+				// Can't use this plugin
+				UE_LOG(LogCompile, Log, TEXT("Unable to determine output directory for %s. Cannot export script glue with \'%s\'"), *GeneratedCodeModuleName, *ScriptGenerator->GetGeneratorName());
+				bSupportedPlugin = false;
+			}
+		}
+		if (!bSupportedPlugin)
+		{
+			UE_LOG(LogCompile, Log, TEXT("Script generator \'%s\' not supported for target: %s"), *ScriptGenerator->GetGeneratorName(), *GManifest.TargetName);
+			ScriptPlugins.RemoveAt(PluginIndex);
+		}
+	}
+}
+
+/**
+ * Tries to resolve super classes for classes defined in the given
+ * module.
+ *
+ * @param Package Modules package.
+ */
+void ResolveSuperClasses(UPackage* Package)
+{
+	TArray<UObject*> Objects;
+	GetObjectsWithPackage(Package, Objects);
+
+	for (UObject* Object : Objects)
+	{
+		if (!Object->IsA<UClass>() || Object->HasAnyFlags(RF_ClassDefaultObject))
+		{
+			continue;
+		}
+
+		UClass* DefinedClass = Cast<UClass>(Object);
+
+		if (DefinedClass->HasAnyClassFlags(CLASS_Intrinsic | CLASS_NoExport))
+		{
+			continue;
+		}
+
+		const FSimplifiedParsingClassInfo& ParsingInfo = GTypeDefinitionInfoMap[DefinedClass]->GetUnrealSourceFile().GetDefinedClassParsingInfo(DefinedClass);
+
+		const FString& BaseClassName         = ParsingInfo.GetBaseClassName();
+		const FString& BaseClassNameStripped = GetClassNameWithPrefixRemoved(BaseClassName);
+
+		if (!BaseClassNameStripped.IsEmpty() && !DefinedClass->GetSuperClass())
+		{
+			UClass* FoundBaseClass = FindObject<UClass>(Package, *BaseClassNameStripped);
+
+			if (FoundBaseClass == nullptr)
+			{
+				FoundBaseClass = FindObject<UClass>(ANY_PACKAGE, *BaseClassNameStripped);
+			}
+
+			if (FoundBaseClass == nullptr)
+			{
+				// Don't know its parent class. Raise error.
+				FError::Throwf(TEXT("Couldn't find parent type for '%s' named '%s' in current module (Package: %s) or any other module parsed so far."), *DefinedClass->GetName(), *BaseClassName, *GetNameSafe(Package));
+			}
+
+			DefinedClass->SetSuperStruct(FoundBaseClass);
+			DefinedClass->ClassCastFlags |= FoundBaseClass->ClassCastFlags;
+		}
+	}
+}
+
+ECompilationResult::Type PreparseModules(const FString& ModuleInfoPath, int32& NumFailures)
+{
+	// Three passes.  1) Public 'Classes' headers (legacy)  2) Public headers   3) Private headers
+	enum EHeaderFolderTypes
+	{
+		PublicClassesHeaders = 0,
+		PublicHeaders = 1,
+		PrivateHeaders,
+
+		FolderType_Count
+	};
+
+	ECompilationResult::Type Result = ECompilationResult::Succeeded;
+
+#if !PLATFORM_EXCEPTIONS_DISABLED
+	FGraphEventArray ExceptionTasks;
+
+	auto LogException = [&Result, &NumFailures, &ExceptionTasks](FString&& Filename, int32 Line, const FString& Message)
+	{
+		auto LogExceptionTask = [&Result, &NumFailures, Filename = MoveTemp(Filename), Line, Message]()
+		{
+			TGuardValue<ELogTimes::Type> DisableLogTimes(GPrintLogTimes, ELogTimes::None);
+
+			FString FormattedErrorMessage = FString::Printf(TEXT("%s(%d): Error: %s\r\n"), *Filename, Line, *Message);
+			Result = ECompilationResult::OtherCompilationError;
+
+			UE_LOG(LogCompile, Log, TEXT("%s"), *FormattedErrorMessage);
+			GWarn->Log(ELogVerbosity::Error, FormattedErrorMessage);
+
+			++NumFailures;
+		};
+
+		if (IsInGameThread())
+		{
+			LogExceptionTask();
+		}
+		else
+		{
+			FGraphEventRef EventRef = FFunctionGraphTask::CreateAndDispatchWhenReady(MoveTemp(LogExceptionTask), TStatId(), nullptr, ENamedThreads::GameThread);
+
+			static FCriticalSection ExceptionCS;
+			FScopeLock Lock(&ExceptionCS);
+			ExceptionTasks.Add(EventRef);
+		}
+	};
+#endif
+
+	for (FManifestModule& Module : GManifest.Modules)
+	{
+		if (Result != ECompilationResult::Succeeded)
+		{
+			break;
+		}
+
+		// Force regeneration of all subsequent modules, otherwise data will get corrupted.
+		Module.ForceRegeneration();
+
+		UPackage* Package = Cast<UPackage>(StaticFindObjectFast(UPackage::StaticClass(), NULL, FName(*Module.LongPackageName), false, false));
+		if (Package == NULL)
+		{
+			Package = CreatePackage(NULL, *Module.LongPackageName);
+		}
+		// Set some package flags for indicating that this package contains script
+		// NOTE: We do this even if we didn't have to create the package, because CoreUObject is compiled into UnrealHeaderTool and we still
+		//       want to make sure our flags get set
+		Package->SetPackageFlags(PKG_ContainsScript | PKG_Compiling);
+		Package->ClearPackageFlags(PKG_ClientOptional | PKG_ServerSideOnly);
+		
+		if(Module.OverrideModuleType == EPackageOverrideType::None)
+		{
+			switch (Module.ModuleType)
+			{
+			case EBuildModuleType::GameEditor:
+			case EBuildModuleType::EngineEditor:
+				Package->SetPackageFlags(PKG_EditorOnly);
+				break;
+
+			case EBuildModuleType::GameDeveloper:
+			case EBuildModuleType::EngineDeveloper:
+				Package->SetPackageFlags(PKG_Developer);
+				break;
+
+			case EBuildModuleType::GameUncooked:
+			case EBuildModuleType::EngineUncooked:
+				Package->SetPackageFlags(PKG_UncookedOnly);
+				break;
+			}
+		}
+		else
+		{
+			// If the user has specified this module to have another package flag, then OR it on
+			switch (Module.OverrideModuleType)
+			{
+			case EPackageOverrideType::EditorOnly:
+				Package->SetPackageFlags(PKG_EditorOnly);
+				break;
+
+			case EPackageOverrideType::EngineDeveloper:
+			case EPackageOverrideType::GameDeveloper:
+				Package->SetPackageFlags(PKG_Developer);
+				break;
+
+			case EPackageOverrideType::EngineUncookedOnly:
+			case EPackageOverrideType::GameUncookedOnly:
+				Package->SetPackageFlags(PKG_UncookedOnly);
+				break;
+			}
+		}
+
+		// Add new module or overwrite whatever we had loaded, that data is obsolete.
+		GPackageToManifestModuleMap.Add(Package, &Module);
+
+		double ThisModulePreparseTime = 0.0;
+		int32 NumHeadersPreparsed = 0;
+		FDurationTimer ThisModuleTimer(ThisModulePreparseTime);
+		ThisModuleTimer.Start();
+
+		// Pre-parse the headers
+		for (int32 PassIndex = 0; PassIndex < FolderType_Count && Result == ECompilationResult::Succeeded; ++PassIndex)
+		{
+			EHeaderFolderTypes CurrentlyProcessing = (EHeaderFolderTypes)PassIndex;
+
+			// We'll make an ordered list of all UObject headers we care about.
+			// @todo uht: Ideally 'dependson' would not be allowed from public -> private, or NOT at all for new style headers
+			const TArray<FString>& UObjectHeaders =
+				(CurrentlyProcessing == PublicClassesHeaders) ? Module.PublicUObjectClassesHeaders :
+				(CurrentlyProcessing == PublicHeaders       ) ? Module.PublicUObjectHeaders        :
+				                                                Module.PrivateUObjectHeaders;
+			if (!UObjectHeaders.Num())
+			{
+				continue;
+			}
+
+			NumHeadersPreparsed += UObjectHeaders.Num();
+
+			TArray<FString> HeaderFiles;
+			HeaderFiles.SetNum(UObjectHeaders.Num());
+
+			{
+				SCOPE_SECONDS_COUNTER_UHT(LoadHeaderContentFromFile);
+				ParallelFor(UObjectHeaders.Num(), [&](int32 Index)
+				{
+					const FString& RawFilename = UObjectHeaders[Index];
+
+#if !PLATFORM_EXCEPTIONS_DISABLED
+					try
+#endif
+					{
+						const FString FullFilename = FPaths::ConvertRelativePathToFull(ModuleInfoPath, RawFilename);
+
+						if (!FFileHelper::LoadFileToString(HeaderFiles[Index], *FullFilename))
+						{
+							FError::Throwf(TEXT("UnrealHeaderTool was unable to load source file '%s'"), *FullFilename);
+						}
+					}
+#if !PLATFORM_EXCEPTIONS_DISABLED
+					catch (TCHAR* ErrorMsg)
+					{
+						FString AbsFilename = IFileManager::Get().ConvertToAbsolutePathForExternalAppForRead(*RawFilename);
+						LogException(MoveTemp(AbsFilename), 1, ErrorMsg);
+					}				
+#endif
+				});
+			}
+
+#if !PLATFORM_EXCEPTIONS_DISABLED
+			FTaskGraphInterface::Get().WaitUntilTasksComplete(ExceptionTasks);
+#endif
+
+			if (Result != ECompilationResult::Succeeded)
+			{
+				continue;
+			}
+
+			TArray<FPerHeaderData> PerHeaderData;
+			PerHeaderData.SetNum(UObjectHeaders.Num());
+
+			ParallelFor(UObjectHeaders.Num(), [&](int32 Index)
+			{
+				const FString& RawFilename = UObjectHeaders[Index];
+
+#if !PLATFORM_EXCEPTIONS_DISABLED
+				try
+#endif
+				{
+					PerformSimplifiedClassParse(Package, *RawFilename, *HeaderFiles[Index], PerHeaderData[Index]);
+				}
+#if !PLATFORM_EXCEPTIONS_DISABLED
+				catch (const FFileLineException& Ex)
+				{
+					FString AbsFilename = IFileManager::Get().ConvertToAbsolutePathForExternalAppForRead(*Ex.Filename);
+					LogException(MoveTemp(AbsFilename), Ex.Line, Ex.Message);
+				}
+				catch (TCHAR* ErrorMsg)
+				{
+					FString AbsFilename = IFileManager::Get().ConvertToAbsolutePathForExternalAppForRead(*RawFilename);
+					LogException(MoveTemp(AbsFilename), 1, ErrorMsg);
+				}
+#endif
+			});
+
+#if !PLATFORM_EXCEPTIONS_DISABLED
+			FTaskGraphInterface::Get().WaitUntilTasksComplete(ExceptionTasks);
+#endif
+
+			if (Result != ECompilationResult::Succeeded)
+			{
+				continue;
+			}
+
+			for (int32 Index = 0; Index < UObjectHeaders.Num(); ++Index)
+			{
+				const FString& RawFilename = UObjectHeaders[Index];
+
+#if !PLATFORM_EXCEPTIONS_DISABLED
+				try
+#endif
+				{
+					// Import class.
+					const FString FullFilename = FPaths::ConvertRelativePathToFull(ModuleInfoPath, RawFilename);
+
+					ProcessInitialClassParse(PerHeaderData[Index]);
+					TSharedRef<FUnrealSourceFile> UnrealSourceFile = PerHeaderData[Index].UnrealSourceFile.ToSharedRef();
+					FUnrealSourceFile* UnrealSourceFilePtr = &UnrealSourceFile.Get();
+					GUnrealSourceFilesMap.Add(FPaths::GetCleanFilename(RawFilename), UnrealSourceFile);
+
+					if (CurrentlyProcessing == PublicClassesHeaders)
+					{
+						GPublicSourceFileSet.Add(UnrealSourceFilePtr);
+					}
+
+					// Save metadata for the class path, both for it's include path and relative to the module base directory
+					if (FullFilename.StartsWith(Module.BaseDirectory))
+					{
+						// Get the path relative to the module directory
+						const TCHAR* ModuleRelativePath = *FullFilename + Module.BaseDirectory.Len();
+
+						UnrealSourceFile->SetModuleRelativePath(ModuleRelativePath);
+
+						// Calculate the include path
+						const TCHAR* IncludePath = ModuleRelativePath;
+
+						// Walk over the first potential slash
+						if (*IncludePath == TEXT('/'))
+						{
+							IncludePath++;
+						}
+
+						// Does this module path start with a known include path location? If so, we can cut that part out of the include path
+						static const TCHAR PublicFolderName[]  = TEXT("Public/");
+						static const TCHAR PrivateFolderName[] = TEXT("Private/");
+						static const TCHAR ClassesFolderName[] = TEXT("Classes/");
+						if (FCString::Strnicmp(IncludePath, PublicFolderName, UE_ARRAY_COUNT(PublicFolderName) - 1) == 0)
+						{
+							IncludePath += (UE_ARRAY_COUNT(PublicFolderName) - 1);
+						}
+						else if (FCString::Strnicmp(IncludePath, PrivateFolderName, UE_ARRAY_COUNT(PrivateFolderName) - 1) == 0)
+						{
+							IncludePath += (UE_ARRAY_COUNT(PrivateFolderName) - 1);
+						}
+						else if (FCString::Strnicmp(IncludePath, ClassesFolderName, UE_ARRAY_COUNT(ClassesFolderName) - 1) == 0)
+						{
+							IncludePath += (UE_ARRAY_COUNT(ClassesFolderName) - 1);
+						}
+
+						// Add the include path
+						if (*IncludePath != 0)
+						{
+							UnrealSourceFile->SetIncludePath(MoveTemp(IncludePath));
+						}
+					}
+				}
+#if !PLATFORM_EXCEPTIONS_DISABLED
+				catch (const FFileLineException& Ex)
+				{
+					FString AbsFilename = IFileManager::Get().ConvertToAbsolutePathForExternalAppForRead(*Ex.Filename);
+					LogException(MoveTemp(AbsFilename), Ex.Line, Ex.Message);
+				}
+				catch (TCHAR* ErrorMsg)
+				{
+					FString AbsFilename = IFileManager::Get().ConvertToAbsolutePathForExternalAppForRead(*RawFilename);
+					LogException(MoveTemp(AbsFilename), 1, ErrorMsg);
+				}
+#endif
+			}
+			if (Result == ECompilationResult::Succeeded && NumFailures != 0)
+			{
+				Result = ECompilationResult::OtherCompilationError;
+			}
+		}
+
+		// Don't resolve superclasses for module when loading from makefile.
+		// Data is only partially loaded at this point.
+#if !PLATFORM_EXCEPTIONS_DISABLED
+		try
+#endif
+		{
+			ResolveSuperClasses(Package);
+		}
+#if !PLATFORM_EXCEPTIONS_DISABLED
+		catch (TCHAR* ErrorMsg)
+		{
+			TGuardValue<ELogTimes::Type> DisableLogTimes(GPrintLogTimes, ELogTimes::None);
+
+			FString FormattedErrorMessage = FString::Printf(TEXT("Error: %s\r\n"), ErrorMsg);
+
+			Result = GCompilationResult;
+
+			UE_LOG(LogCompile, Log, TEXT("%s"), *FormattedErrorMessage);
+			GWarn->Log(ELogVerbosity::Error, FormattedErrorMessage);
+
+			++NumFailures;
+		}
+#endif
+
+		ThisModuleTimer.Stop();
+		UE_LOG(LogCompile, Log, TEXT("Preparsed module %s containing %i files(s) in %.2f secs."), *Module.LongPackageName, NumHeadersPreparsed, ThisModulePreparseTime);
+	}
+
+	return Result;
+}
+
+ECompilationResult::Type UnrealHeaderTool_Main(const FString& ModuleInfoFilename)
+{
+	double MainTime = 0.0;
+	FDurationTimer MainTimer(MainTime);
+	MainTimer.Start();
+
+	check(GIsUCCMakeStandaloneHeaderGenerator);
+	ECompilationResult::Type Result = ECompilationResult::Succeeded;
+
+	FString ModuleInfoPath = FPaths::GetPath(ModuleInfoFilename);
+
+	// Load the manifest file, giving a list of all modules to be processed, pre-sorted by dependency ordering
+#if !PLATFORM_EXCEPTIONS_DISABLED
+	try
+#endif
+	{
+		GManifest = FManifest::LoadFromFile(ModuleInfoFilename);
+	}
+#if !PLATFORM_EXCEPTIONS_DISABLED
+	catch (const TCHAR* Ex)
+	{
+		UE_LOG(LogCompile, Error, TEXT("Failed to load manifest file '%s': %s"), *ModuleInfoFilename, Ex);
+		return GCompilationResult;
+	}
+#endif
+
+	// Counters.
+	int32 NumFailures = 0;
+	double TotalModulePreparseTime = 0.0;
+	double TotalParseAndCodegenTime = 0.0;
+
+	{
+		FDurationTimer TotalModulePreparseTimer(TotalModulePreparseTime);
+		TotalModulePreparseTimer.Start();
+		Result = PreparseModules(ModuleInfoPath, NumFailures);
+		TotalModulePreparseTimer.Stop();
+	}
+	// Do the actual parse of the headers and generate for them
+	if (Result == ECompilationResult::Succeeded)
+	{
+		FScopedDurationTimer ParseAndCodeGenTimer(TotalParseAndCodegenTime);
+
+		TMap<UPackage*, TArray<UClass*>> ClassesByPackageMap;
+		ClassesByPackageMap.Reserve(GManifest.Modules.Num());
+
+		// Verify that all script declared superclasses exist.
+		for (UClass* ScriptClass : TObjectRange<UClass>())
+		{
+			ClassesByPackageMap.FindOrAdd(ScriptClass->GetOutermost()).Add(ScriptClass);
+
+			const UClass* ScriptSuperClass = ScriptClass->GetSuperClass();
+
+			if (ScriptSuperClass && !ScriptSuperClass->HasAnyClassFlags(CLASS_Intrinsic) && GTypeDefinitionInfoMap.Contains(ScriptClass) && !GTypeDefinitionInfoMap.Contains(ScriptSuperClass))
+			{
+				class FSuperClassContextSupplier : public FContextSupplier
+				{
+				public:
+					FSuperClassContextSupplier(const UClass* Class)
+						: DefinitionInfo(GTypeDefinitionInfoMap[Class])
+					{ }
+
+					virtual FString GetContext() override
+					{
+						FString Filename = IFileManager::Get().ConvertToAbsolutePathForExternalAppForRead(*DefinitionInfo->GetUnrealSourceFile().GetFilename());
+						int32 LineNumber = DefinitionInfo->GetLineNumber();
+						return FString::Printf(TEXT("%s(%i)"), *Filename, LineNumber);
+					}
+				private:
+					TSharedRef<FUnrealTypeDefinitionInfo> DefinitionInfo;
+				} ContextSupplier(ScriptClass);
+
+				FContextSupplier* OldContext = GWarn->GetContext();
+
+				TGuardValue<ELogTimes::Type> DisableLogTimes(GPrintLogTimes, ELogTimes::None);
+
+				GWarn->SetContext(&ContextSupplier);
+				GWarn->Log(ELogVerbosity::Error, FString::Printf(TEXT("Error: Superclass %s of class %s not found"), *ScriptSuperClass->GetName(), *ScriptClass->GetName()));
+				GWarn->SetContext(OldContext);
+
+				Result = ECompilationResult::OtherCompilationError;
+				++NumFailures;
+			}
+		}
+
+		if (Result == ECompilationResult::Succeeded)
+		{
+			TArray<IScriptGeneratorPluginInterface*> ScriptPlugins;
+			// Can only export scripts for game targets
+			if (GManifest.IsGameTarget)
+			{
+				GetScriptPlugins(ScriptPlugins);
+			}
+
+			for (const FManifestModule& Module : GManifest.Modules)
+			{
+				if (UPackage* Package = Cast<UPackage>(StaticFindObjectFast(UPackage::StaticClass(), NULL, FName(*Module.LongPackageName), false, false)))
+				{
+					FClasses AllClasses(ClassesByPackageMap.Find(Package));
+					AllClasses.Validate();
+
+					Result = FHeaderParser::ParseAllHeadersInside(AllClasses, GWarn, Package, Module, ScriptPlugins);
+					if (Result != ECompilationResult::Succeeded)
+					{
+						++NumFailures;
+						break;
+					}
+				}
+			}
+
+			{
+				FScopedDurationTimer PluginTimeTracker(GPluginOverheadTime);
+				for (IScriptGeneratorPluginInterface* ScriptGenerator : ScriptPlugins)
+				{
+					ScriptGenerator->FinishExport();
+				}
+			}
+
+			// Get a list of external dependencies from each enabled plugin
+			FString ExternalDependencies;
+			for (IScriptGeneratorPluginInterface* ScriptPlugin : ScriptPlugins)
+			{
+				TArray<FString> PluginExternalDependencies;
+				ScriptPlugin->GetExternalDependencies(PluginExternalDependencies);
+
+				for (const FString& PluginExternalDependency : PluginExternalDependencies)
+				{
+					ExternalDependencies += PluginExternalDependency + LINE_TERMINATOR;
+				}
+			}
+			FFileHelper::SaveStringToFile(ExternalDependencies, *GManifest.ExternalDependenciesFile);
+		}
+	}
+
+	// Avoid TArray slack for meta data.
+	GScriptHelper.Shrink();
+
+	// Finish all async file tasks before stopping the clock
+	FTaskGraphInterface::Get().WaitUntilTasksComplete(GAsyncFileTasks);
+
+	MainTimer.Stop();
+
+	UE_LOG(LogCompile, Log, TEXT("Preparsing %i modules took %.2f seconds"), GManifest.Modules.Num(), TotalModulePreparseTime);
+	UE_LOG(LogCompile, Log, TEXT("Parsing took %.2f seconds"), TotalParseAndCodegenTime - GHeaderCodeGenTime);
+	UE_LOG(LogCompile, Log, TEXT("Code generation took %.2f seconds"), GHeaderCodeGenTime);
+	UE_LOG(LogCompile, Log, TEXT("ScriptPlugin overhead was %.2f seconds"), GPluginOverheadTime);
+	UE_LOG(LogCompile, Log, TEXT("Macroize time was %.2f seconds"), GMacroizeTime);
+
+	FUnrealHeaderToolStats& Stats = FUnrealHeaderToolStats::Get();
+	for (const TPair<FName, double>& Pair : Stats.Counters)
+	{
+		FString CounterName = Pair.Key.ToString();
+		UE_LOG(LogCompile, Log, TEXT("%s timer was %.3f seconds"), *CounterName, Pair.Value);
+	}
+
+	UE_LOG(LogCompile, Log, TEXT("Total time was %.2f seconds"), MainTime);
+
+	if (bWriteContents)
+	{
+		UE_LOG(LogCompile, Log, TEXT("********************************* Wrote reference generated code to ReferenceGeneratedCode."));
+	}
+	else if (bVerifyContents)
+	{
+		UE_LOG(LogCompile, Log, TEXT("********************************* Wrote generated code to VerifyGeneratedCode and compared to ReferenceGeneratedCode"));
+		for (FString& Msg : ChangeMessages)
+		{
+			UE_LOG(LogCompile, Error, TEXT("%s"), *Msg);
+		}
+		TArray<FString> RefFileNames;
+		IFileManager::Get().FindFiles( RefFileNames, *(FPaths::ProjectSavedDir() / TEXT("ReferenceGeneratedCode/*.*")), true, false );
+		TArray<FString> VerFileNames;
+		IFileManager::Get().FindFiles( VerFileNames, *(FPaths::ProjectSavedDir() / TEXT("VerifyGeneratedCode/*.*")), true, false );
+		if (RefFileNames.Num() != VerFileNames.Num())
+		{
+			UE_LOG(LogCompile, Error, TEXT("Number of generated files mismatch ref=%d, ver=%d"), RefFileNames.Num(), VerFileNames.Num());
+		}
+	}
+
+	RequestEngineExit(TEXT("UnrealHeaderTool finished"));
+
+	if (Result != ECompilationResult::Succeeded || NumFailures > 0)
+	{
+		return ECompilationResult::OtherCompilationError;
+	}
+
+	return Result;
+}
+
+UClass* ProcessParsedClass(bool bClassIsAnInterface, TArray<FHeaderProvider>& DependentOn, const FString& ClassName, const FString& BaseClassName, UObject* InParent, EObjectFlags Flags)
+{
+	FString ClassNameStripped = GetClassNameWithPrefixRemoved(*ClassName);
+
+	// All classes must start with a valid unreal prefix
+	if (!FHeaderParser::ClassNameHasValidPrefix(ClassName, ClassNameStripped))
+	{
+		FError::Throwf(TEXT("Invalid class name '%s'. The class name must have an appropriate prefix added (A for Actors, U for other classes)."), *ClassName);
+	}
+
+	if(FHeaderParser::IsReservedTypeName(ClassNameStripped))
+	{
+		FError::Throwf(TEXT("Invalid class name '%s'. Cannot use a reserved name ('%s')."), *ClassName, *ClassNameStripped);
+	}
+
+	// Ensure the base class has any valid prefix and exists as a valid class. Checking for the 'correct' prefix will occur during compilation
+	FString BaseClassNameStripped;
+	if (!BaseClassName.IsEmpty())
+	{
+		BaseClassNameStripped = GetClassNameWithPrefixRemoved(BaseClassName);
+		if (!FHeaderParser::ClassNameHasValidPrefix(BaseClassName, BaseClassNameStripped))
+		{
+			FError::Throwf(TEXT("No prefix or invalid identifier for base class %s.\nClass names must match Unreal prefix specifications (e.g., \"UObject\" or \"AActor\")"), *BaseClassName);
+		}
+	}
+
+	//UE_LOG(LogCompile, Log, TEXT("Class: %s extends %s"),*ClassName,*BaseClassName);
+	// Handle failure and non-class headers.
+	if (BaseClassName.IsEmpty() && (ClassName != TEXT("UObject")))
+	{
+		FError::Throwf(TEXT("Class '%s' must inherit UObject or a UObject-derived class"), *ClassName);
+	}
+
+	if (ClassName == BaseClassName)
+	{
+		FError::Throwf(TEXT("Class '%s' cannot inherit from itself"), *ClassName);
+	}
+
+	// In case the file system and the class disagree on the case of the
+	// class name replace the fname with the one from the script class file
+	// This is needed because not all source control systems respect the
+	// original filename's case
+	FName ClassNameReplace(*ClassName, FNAME_Replace_Not_Safe_For_Threading);
+
+	// Use stripped class name for processing and replace as we did above
+	FName ClassNameStrippedReplace(*ClassNameStripped, FNAME_Replace_Not_Safe_For_Threading);
+
+	UClass* ResultClass = FindObject<UClass>(InParent, *ClassNameStripped);
+
+	// if we aren't generating headers, then we shouldn't set misaligned object, since it won't get cleared
+
+	const static bool bVerboseOutput = FParse::Param(FCommandLine::Get(), TEXT("VERBOSE"));
+
+	if (ResultClass == nullptr || !ResultClass->IsNative())
+	{
+		// detect if the same class name is used in multiple packages
+		if (ResultClass == nullptr)
+		{
+			UClass* ConflictingClass = FindObject<UClass>(ANY_PACKAGE, *ClassNameStripped, true);
+			if (ConflictingClass != nullptr)
+			{
+				UE_LOG_WARNING_UHT(TEXT("Duplicate class name: %s also exists in file %s"), *ClassName, *ConflictingClass->GetOutermost()->GetName());
+			}
+		}
+
+		// Create new class.
+		ResultClass = new(EC_InternalUseOnlyConstructor, InParent, *ClassNameStripped, Flags) UClass(FObjectInitializer(), nullptr);
+
+		// add CLASS_Interface flag if the class is an interface
+		// NOTE: at this pre-parsing/importing stage, we cannot know if our super class is an interface or not,
+		// we leave the validation to the main header parser
+		if (bClassIsAnInterface)
+		{
+			ResultClass->ClassFlags |= CLASS_Interface;
+		}
+
+		if (bVerboseOutput)
+		{
+			UE_LOG(LogCompile, Log, TEXT("Imported: %s"), *ResultClass->GetFullName());
+		}
+	}
+
+	if (bVerboseOutput)
+	{
+		for (const FHeaderProvider& Dependency : DependentOn)
+		{
+			UE_LOG(LogCompile, Log, TEXT("\tAdding %s as a dependency"), *Dependency.ToString());
+		}
+	}
+
+	return ResultClass;
+}
+
+void PerformSimplifiedClassParse(UPackage* InParent, const TCHAR* FileName, const TCHAR* Buffer, FPerHeaderData& PerHeaderData)
+{
+	// Parse the header to extract the information needed
+	FUHTStringBuilder ClassHeaderTextStrippedOfCppText;
+
+	FHeaderParser::SimplifiedClassParse(FileName, Buffer, /*out*/ PerHeaderData.ParsedClassArray, /*out*/ PerHeaderData.DependsOn, ClassHeaderTextStrippedOfCppText);
+
+	FUnrealSourceFile* UnrealSourceFilePtr = new FUnrealSourceFile(InParent, FileName, MoveTemp(ClassHeaderTextStrippedOfCppText));
+	PerHeaderData.UnrealSourceFile = MakeShareable(UnrealSourceFilePtr);
+}
+
+void ProcessInitialClassParse(FPerHeaderData& PerHeaderData)
+{
+	TSharedRef<FUnrealSourceFile> UnrealSourceFile = PerHeaderData.UnrealSourceFile.ToSharedRef();
+	UPackage* InParent = UnrealSourceFile->GetPackage();
+	for (FSimplifiedParsingClassInfo& ParsedClassInfo : PerHeaderData.ParsedClassArray)
+	{
+		UClass* ResultClass = ProcessParsedClass(ParsedClassInfo.IsInterface(), PerHeaderData.DependsOn, ParsedClassInfo.GetClassName(), ParsedClassInfo.GetBaseClassName(), InParent, RF_Public | RF_Standalone);
+		GStructToSourceLine.Add(ResultClass, MakeTuple(UnrealSourceFile, ParsedClassInfo.GetClassDefLine()));
+
+		FScope::AddTypeScope(ResultClass, &UnrealSourceFile->GetScope().Get());
+
+		GTypeDefinitionInfoMap.Add(ResultClass, MakeShared<FUnrealTypeDefinitionInfo>(UnrealSourceFile.Get(), ParsedClassInfo.GetClassDefLine()));
+		UnrealSourceFile->AddDefinedClass(ResultClass, MoveTemp(ParsedClassInfo));
+	}
+
+	for (FHeaderProvider& DependsOnElement : PerHeaderData.DependsOn)
+	{
+		UnrealSourceFile->GetIncludes().AddUnique(DependsOnElement);
+	}
+}
