Branching UHT to UHTLite:

diff --git a/Engine/Source/Programs/UHTLite/Private/FileLineException.cpp b/Engine/Source/Programs/UHTLite/Private/FileLineException.cpp
new file mode 100644
--- /dev/null
+++ b/Engine/Source/Programs/UHTLite/Private/FileLineException.cpp
@@ -0,0 +1,45 @@
+// Copyright Epic Games, Inc. All Rights Reserved.
+
+#include "FileLineException.h"
+#include "UnrealHeaderTool.h"
+
+void VARARGS FFileLineException::ThrowfImpl(FString&& InFilename, int32 InLine, const TCHAR* Fmt, ...)
+{
+	int32	BufferSize	= 512;
+	TCHAR	StartingBuffer[512];
+	TCHAR*	Buffer		= StartingBuffer;
+	int32	Result		= -1;
+
+	// First try to print to a stack allocated location 
+	GET_VARARGS_RESULT( Buffer, BufferSize, BufferSize-1, Fmt, Fmt, Result );
+
+	// If that fails, start allocating regular memory
+	if( Result == -1 )
+	{
+		Buffer = nullptr;
+		while(Result == -1)
+		{
+			BufferSize *= 2;
+			Buffer = (TCHAR*) FMemory::Realloc( Buffer, BufferSize * sizeof(TCHAR) );
+			GET_VARARGS_RESULT( Buffer, BufferSize, BufferSize-1, Fmt, Fmt, Result );
+		}
+	}
+
+	Buffer[Result] = 0;
+
+	FString ResultString(Buffer);
+
+	if( BufferSize != 512 )
+	{
+		FMemory::Free( Buffer );
+	}
+
+	throw FFileLineException(MoveTemp(ResultString), MoveTemp(InFilename), InLine);
+}
+
+FFileLineException::FFileLineException(FString&& InMessage, FString&& InFilename, int32 InLine)
+	: Message (MoveTemp(InMessage))
+	, Filename(MoveTemp(InFilename))
+	, Line    (InLine)
+{
+}
diff --git a/Engine/Source/Programs/UHTLite/Private/FileLineException.h b/Engine/Source/Programs/UHTLite/Private/FileLineException.h
new file mode 100644
--- /dev/null
+++ b/Engine/Source/Programs/UHTLite/Private/FileLineException.h
@@ -0,0 +1,27 @@
+// Copyright Epic Games, Inc. All Rights Reserved.
+
+#pragma once
+
+#include "CoreMinimal.h"
+#include "Containers/UnrealString.h"
+#include "Templates/IsValidVariadicFunctionArg.h"
+
+struct FFileLineException
+{
+	FString Message;
+	FString Filename;
+	int32   Line;
+
+	template <typename... Types>
+	UE_NORETURN static void VARARGS Throwf(FString&& Filename, int32 Line, const TCHAR* Fmt, Types... Args)
+	{
+		static_assert(TAnd<TIsValidVariadicFunctionArg<Types>...>::Value, "Invalid argument(s) passed to FError::Throwf");
+
+		ThrowfImpl(MoveTemp(Filename), Line, Fmt, Args...);
+	}
+
+private:
+	UE_NORETURN static void VARARGS ThrowfImpl(FString&& Filename, int32 Line, const TCHAR* Fmt, ...);
+
+	FFileLineException(FString&& InMessage, FString&& InFilename, int32 InLine);
+};
diff --git a/Engine/Source/Programs/UHTLite/Private/GeneratedCodeVersion.h b/Engine/Source/Programs/UHTLite/Private/GeneratedCodeVersion.h
new file mode 100644
--- /dev/null
+++ b/Engine/Source/Programs/UHTLite/Private/GeneratedCodeVersion.h
@@ -0,0 +1,56 @@
+// Copyright Epic Games, Inc. All Rights Reserved.
+#pragma once
+
+#include "CoreMinimal.h"
+#include "Containers/UnrealString.h"
+
+class FArchive;
+class FString;
+
+//
+// This MUST be kept in sync with EGeneratedBodyVersion in UBT defined in ExternalExecution.cs
+// and with ToGeneratedBodyVersion function below
+//
+enum class EGeneratedCodeVersion : uint8
+{
+	None,
+	V1,
+	V2,
+	VLatest = V2
+};
+
+inline FArchive& operator<<(FArchive& Ar, EGeneratedCodeVersion& Type)
+{
+	if (Ar.IsLoading())
+	{
+		uint8 Value;
+		Ar << Value;
+		Type = (EGeneratedCodeVersion)Value;
+	}
+	else if (Ar.IsSaving())
+	{
+		uint8 Value = (uint8)Type;
+		Ar << Value;
+	}
+	return Ar;
+}
+
+inline EGeneratedCodeVersion ToGeneratedCodeVersion(const FString& InString)
+{
+	if (InString.Compare(TEXT("V1")) == 0)
+	{
+		return EGeneratedCodeVersion::V1;
+	}
+
+	if (InString.Compare(TEXT("V2")) == 0)
+	{
+		return EGeneratedCodeVersion::V2;
+	}
+
+	if (InString.Compare(TEXT("VLatest")) == 0)
+	{
+		return EGeneratedCodeVersion::VLatest;
+	}
+
+	return EGeneratedCodeVersion::None;
+}
