Branching UHT to UHTLite:

diff --git a/Engine/Source/Programs/UHTLite/Private/HeaderParser.cpp b/Engine/Source/Programs/UHTLite/Private/HeaderParser.cpp
new file mode 100644
--- /dev/null
+++ b/Engine/Source/Programs/UHTLite/Private/HeaderParser.cpp
@@ -0,0 +1,10733 @@
+// Copyright Epic Games, Inc. All Rights Reserved.
+
+
+#include "HeaderParser.h"
+#include "UnrealHeaderTool.h"
+#include "HAL/FileManager.h"
+#include "Misc/CommandLine.h"
+#include "Misc/ConfigCacheIni.h"
+#include "Misc/FeedbackContext.h"
+#include "UObject/Interface.h"
+#include "ParserClass.h"
+#include "GeneratedCodeVersion.h"
+#include "ClassDeclarationMetaData.h"
+#include "ProfilingDebugging/ScopedTimers.h"
+#include "NativeClassExporter.h"
+#include "Classes.h"
+#include "StringUtils.h"
+#include "Misc/DefaultValueHelper.h"
+#include "Manifest.h"
+#include "Math/UnitConversion.h"
+#include "FileLineException.h"
+#include "UnrealTypeDefinitionInfo.h"
+#include "Containers/EnumAsByte.h"
+#include "Algo/AllOf.h"
+#include "Algo/Find.h"
+#include "Algo/FindSortedStringCaseInsensitive.h"
+#include "Misc/ScopeExit.h"
+
+#include "Specifiers/CheckedMetadataSpecifiers.h"
+#include "Specifiers/FunctionSpecifiers.h"
+#include "Specifiers/InterfaceSpecifiers.h"
+#include "Specifiers/StructSpecifiers.h"
+#include "Specifiers/VariableSpecifiers.h"
+
+double GPluginOverheadTime = 0.0;
+double GHeaderCodeGenTime = 0.0;
+
+/*-----------------------------------------------------------------------------
+	Constants & declarations.
+-----------------------------------------------------------------------------*/
+
+/**
+ * Data struct that annotates source files that failed during parsing.
+ */
+class FFailedFilesAnnotation
+{
+public:
+	/**
+	 * Gets annotation state for given source file.
+	 */
+	bool Get(FUnrealSourceFile* SourceFile) const
+	{
+		return AnnotatedSet.Contains(SourceFile);
+	}
+
+	/**
+	 * Sets annotation state to true for given source file.
+	 */
+	void Set(FUnrealSourceFile* SourceFile)
+	{
+		AnnotatedSet.Add(SourceFile);
+	}
+
+private:
+	// Annotation set.
+	TSet<FUnrealSourceFile*> AnnotatedSet;
+} static FailedFilesAnnotation;
+
+enum {MAX_ARRAY_SIZE=2048};
+
+namespace
+{
+	static const FName NAME_Comment(TEXT("Comment"));
+	static const FName NAME_ToolTip(TEXT("ToolTip"));
+	static const FName NAME_DocumentationPolicy(TEXT("DocumentationPolicy"));
+	static const FName NAME_AllowPrivateAccess(TEXT("AllowPrivateAccess"));
+	static const FName NAME_ExposeOnSpawn(TEXT("ExposeOnSpawn"));
+	static const FName NAME_NativeConst(TEXT("NativeConst"));
+	static const FName NAME_NativeConstTemplateArg(TEXT("NativeConstTemplateArg"));
+	static const FName NAME_BlueprintInternalUseOnly(TEXT("BlueprintInternalUseOnly"));
+	static const FName NAME_DeprecatedFunction(TEXT("DeprecatedFunction"));
+	static const FName NAME_BlueprintSetter(TEXT("BlueprintSetter"));
+	static const FName NAME_BlueprintGetter(TEXT("BlueprintGetter"));
+	static const FName NAME_Category(TEXT("Category"));
+	static const FName NAME_ReturnValue(TEXT("ReturnValue"));
+	static const FName NAME_CppFromBpEvent(TEXT("CppFromBpEvent"));
+	static const FName NAME_CustomThunk(TEXT("CustomThunk"));
+	static const FName NAME_ArraySizeEnum(TEXT("ArraySizeEnum"));
+	static const FName NAME_ClassGroupNames(TEXT("ClassGroupNames"));
+	static const FName NAME_AutoCollapseCategories(TEXT("AutoCollapseCategories"));
+	static const FName NAME_HideFunctions(TEXT("HideFunctions"));
+	static const FName NAME_AutoExpandCategories(TEXT("AutoExpandCategories"));
+	static const FName NAME_EditInline(TEXT("EditInline"));
+	static const FName NAME_IncludePath(TEXT("IncludePath"));
+	static const FName NAME_ModuleRelativePath(TEXT("ModuleRelativePath"));
+	static const FName NAME_CannotImplementInterfaceInBlueprint(TEXT("CannotImplementInterfaceInBlueprint"));
+	static const FName NAME_UIMin(TEXT("UIMin"));
+	static const FName NAME_UIMax(TEXT("UIMax"));
+	static const FName NAME_BlueprintType(TEXT("BlueprintType"));
+}
+
+const FName FHeaderParserNames::NAME_HideCategories(TEXT("HideCategories"));
+const FName FHeaderParserNames::NAME_ShowCategories(TEXT("ShowCategories"));
+const FName FHeaderParserNames::NAME_SparseClassDataTypes(TEXT("SparseClassDataTypes"));
+const FName FHeaderParserNames::NAME_IsConversionRoot(TEXT("IsConversionRoot"));
+const FName FHeaderParserNames::NAME_AdvancedClassDisplay(TEXT("AdvancedClassDisplay"));
+
+EGeneratedCodeVersion FHeaderParser::DefaultGeneratedCodeVersion = EGeneratedCodeVersion::V1;
+TArray<FString> FHeaderParser::StructsWithNoPrefix;
+TArray<FString> FHeaderParser::StructsWithTPrefix;
+FRigVMStructMap FHeaderParser::StructRigVMMap;
+TArray<FString> FHeaderParser::DelegateParameterCountStrings;
+TMap<FString, FString> FHeaderParser::TypeRedirectMap;
+TArray<FString> FHeaderParser::PropertyCPPTypesRequiringUIRanges = { TEXT("float"), TEXT("double") };
+TArray<FString> FHeaderParser::ReservedTypeNames = { TEXT("none") };
+TMap<UClass*, ClassDefinitionRange> ClassDefinitionRanges;
+
+/**
+ * Dirty hack global variable to allow different result codes passed through
+ * exceptions. Needs to be fixed in future versions of UHT.
+ */
+extern ECompilationResult::Type GCompilationResult;
+
+/*-----------------------------------------------------------------------------
+	Utility functions.
+-----------------------------------------------------------------------------*/
+
+namespace
+{
+	bool ProbablyAMacro(const TCHAR* Identifier)
+	{
+		// Macros must start with a capitalized alphanumeric character or underscore
+		TCHAR FirstChar = Identifier[0];
+		if (FirstChar != TEXT('_') && (FirstChar < TEXT('A') || FirstChar > TEXT('Z')))
+		{
+			return false;
+		}
+
+		// Test for known delegate and event macros.
+		TCHAR MulticastDelegateStart[] = TEXT("DECLARE_MULTICAST_DELEGATE");
+		if (!FCString::Strncmp(Identifier, MulticastDelegateStart, UE_ARRAY_COUNT(MulticastDelegateStart) - 1))
+		{
+			return true;
+		}
+
+		TCHAR DelegateStart[] = TEXT("DECLARE_DELEGATE");
+		if (!FCString::Strncmp(Identifier, DelegateStart, UE_ARRAY_COUNT(DelegateStart) - 1))
+		{
+			return true;
+		}
+
+		TCHAR DelegateEvent[] = TEXT("DECLARE_EVENT");
+		if (!FCString::Strncmp(Identifier, DelegateEvent, UE_ARRAY_COUNT(DelegateEvent) - 1))
+		{
+			return true;
+		}
+
+		// Failing that, we'll guess about it being a macro based on it being a fully-capitalized identifier.
+		while (TCHAR Ch = *++Identifier)
+		{
+			if (Ch != TEXT('_') && (Ch < TEXT('A') || Ch > TEXT('Z')) && (Ch < TEXT('0') || Ch > TEXT('9')))
+			{
+				return false;
+			}
+		}
+
+		return true;
+	}
+
+	/**
+	 * Tests if an identifier looks like a macro which doesn't have a following open parenthesis.
+	 *
+	 * @param HeaderParser  The parser to retrieve the next token.
+	 * @param Token         The token to test for being callable-macro-like.
+	 *
+	 * @return true if it looks like a non-callable macro, false otherwise.
+	 */
+	bool ProbablyAnUnknownObjectLikeMacro(FHeaderParser& HeaderParser, FToken Token)
+	{
+		// Non-identifiers are not macros
+		if (Token.TokenType != TOKEN_Identifier)
+		{
+			return false;
+		}
+
+		// Macros must start with a capitalized alphanumeric character or underscore
+		TCHAR FirstChar = Token.Identifier[0];
+		if (FirstChar != TEXT('_') && (FirstChar < TEXT('A') || FirstChar > TEXT('Z')))
+		{
+			return false;
+		}
+
+		// We'll guess about it being a macro based on it being fully-capitalized with at least one underscore.
+		const TCHAR* IdentPtr = Token.Identifier;
+		int32 UnderscoreCount = 0;
+		while (TCHAR Ch = *++IdentPtr)
+		{
+			if (Ch == TEXT('_'))
+			{
+				++UnderscoreCount;
+			}
+			else if ((Ch < TEXT('A') || Ch > TEXT('Z')) && (Ch < TEXT('0') || Ch > TEXT('9')))
+			{
+				return false;
+			}
+		}
+
+		// We look for at least one underscore as a convenient way of whitelisting many known macros
+		// like FORCEINLINE and CONSTEXPR, and non-macros like FPOV and TCHAR.
+		if (UnderscoreCount == 0)
+		{
+			return false;
+		}
+
+		// Identifiers which end in _API are known
+		if (IdentPtr - Token.Identifier > 4 && IdentPtr[-4] == TEXT('_') && IdentPtr[-3] == TEXT('A') && IdentPtr[-2] == TEXT('P') && IdentPtr[-1] == TEXT('I'))
+		{
+			return false;
+		}
+
+		// Ignore certain known macros or identifiers that look like macros.
+		// IMPORTANT: needs to be in lexicographical order.
+		static const TCHAR* Whitelist[] =
+		{
+			TEXT("FORCEINLINE_DEBUGGABLE"),
+			TEXT("FORCEINLINE_STATS"),
+			TEXT("SIZE_T")
+		};
+		if (Algo::FindSortedStringCaseInsensitive(Token.Identifier, Whitelist, UE_ARRAY_COUNT(Whitelist)) >= 0)
+		{
+			return false;
+		}
+
+		// Check if there's an open parenthesis following the token.
+		//
+		// Rather than ungetting the bracket token, we unget the original identifier token,
+		// then get it again, so we don't lose any comments which may exist between the token 
+		// and the non-bracket.
+		FToken PossibleBracketToken;
+		HeaderParser.GetToken(PossibleBracketToken);
+		HeaderParser.UngetToken(Token);
+		HeaderParser.GetToken(Token);
+
+		bool bResult = PossibleBracketToken.TokenType != TOKEN_Symbol || !PossibleBracketToken.Matches(TEXT('('));
+		return bResult;
+	}
+
+	/**
+	 * Parse and validate an array of identifiers (inside FUNC_NetRequest, FUNC_NetResponse) 
+	 * @param FuncInfo function info for the current function
+	 * @param Identifiers identifiers inside the net service declaration
+	 */
+	void ParseNetServiceIdentifiers(FFuncInfo& FuncInfo, const TArray<FString>& Identifiers)
+	{
+		static const TCHAR IdTag         [] = TEXT("Id");
+		static const TCHAR ResponseIdTag [] = TEXT("ResponseId");
+		static const TCHAR JSBridgePriTag[] = TEXT("Priority");
+
+		for (const FString& Identifier : Identifiers)
+		{
+			const TCHAR* IdentifierPtr = *Identifier;
+
+			if (const TCHAR* Equals = FCString::Strchr(IdentifierPtr, TEXT('=')))
+			{
+				// It's a tag with an argument
+
+				if (FCString::Strnicmp(IdentifierPtr, IdTag, UE_ARRAY_COUNT(IdTag) - 1) == 0)
+				{
+					int32 TempInt = FCString::Atoi(Equals + 1);
+					if (TempInt <= 0 || TempInt > MAX_uint16)
+					{
+						FError::Throwf(TEXT("Invalid network identifier %s for function"), IdentifierPtr);
+					}
+					FuncInfo.RPCId = (uint16)TempInt;
+				}
+				else if (FCString::Strnicmp(IdentifierPtr, ResponseIdTag, UE_ARRAY_COUNT(ResponseIdTag) - 1) == 0 ||
+					FCString::Strnicmp(IdentifierPtr, JSBridgePriTag, UE_ARRAY_COUNT(JSBridgePriTag) - 1) == 0)
+				{
+					int32 TempInt = FCString::Atoi(Equals + 1);
+					if (TempInt <= 0 || TempInt > MAX_uint16)
+					{
+						FError::Throwf(TEXT("Invalid network identifier %s for function"), IdentifierPtr);
+					}
+					FuncInfo.RPCResponseId = (uint16)TempInt;
+				}
+			}
+			else
+			{
+				// Assume it's an endpoint name
+
+				if (FuncInfo.EndpointName.Len())
+				{
+					FError::Throwf(TEXT("Function should not specify multiple endpoints - '%s' found but already using '%s'"), *Identifier);
+				}
+
+				FuncInfo.EndpointName = Identifier;
+			}
+		}
+	}
+
+	/**
+	 * Processes a set of UFUNCTION or UDELEGATE specifiers into an FFuncInfo struct.
+	 *
+	 * @param FuncInfo   - The FFuncInfo object to populate.
+	 * @param Specifiers - The specifiers to process.
+	 */
+	void ProcessFunctionSpecifiers(FFuncInfo& FuncInfo, const TArray<FPropertySpecifier>& Specifiers, TMap<FName, FString>& MetaData)
+	{
+		bool bSpecifiedUnreliable = false;
+		bool bSawPropertyAccessor = false;
+
+		for (const FPropertySpecifier& Specifier : Specifiers)
+		{
+			switch ((EFunctionSpecifier)Algo::FindSortedStringCaseInsensitive(*Specifier.Key, GFunctionSpecifierStrings))
+			{
+				default:
+				{
+					FError::Throwf(TEXT("Unknown function specifier '%s'"), *Specifier.Key);
+				}
+				break;
+
+				case EFunctionSpecifier::BlueprintNativeEvent:
+				{
+					if (FuncInfo.FunctionFlags & FUNC_Net)
+					{
+						UE_LOG_ERROR_UHT(TEXT("BlueprintNativeEvent functions cannot be replicated!") );
+					}
+					else if ( (FuncInfo.FunctionFlags & FUNC_BlueprintEvent) && !(FuncInfo.FunctionFlags & FUNC_Native) )
+					{
+						// already a BlueprintImplementableEvent
+						UE_LOG_ERROR_UHT(TEXT("A function cannot be both BlueprintNativeEvent and BlueprintImplementableEvent!") );
+					}
+					else if (bSawPropertyAccessor)
+					{
+						UE_LOG_ERROR_UHT(TEXT("A function cannot be both BlueprintNativeEvent and a Blueprint Property accessor!"));
+					}
+					else if ( (FuncInfo.FunctionFlags & FUNC_Private) )
+					{
+						UE_LOG_ERROR_UHT(TEXT("A Private function cannot be a BlueprintNativeEvent!") );
+					}
+
+					FuncInfo.FunctionFlags |= FUNC_Event;
+					FuncInfo.FunctionFlags |= FUNC_BlueprintEvent;
+				}
+				break;
+
+				case EFunctionSpecifier::BlueprintImplementableEvent:
+				{
+					if (FuncInfo.FunctionFlags & FUNC_Net)
+					{
+						UE_LOG_ERROR_UHT(TEXT("BlueprintImplementableEvent functions cannot be replicated!") );
+					}
+					else if ( (FuncInfo.FunctionFlags & FUNC_BlueprintEvent) && (FuncInfo.FunctionFlags & FUNC_Native) )
+					{
+						// already a BlueprintNativeEvent
+						UE_LOG_ERROR_UHT(TEXT("A function cannot be both BlueprintNativeEvent and BlueprintImplementableEvent!") );
+					}
+					else if (bSawPropertyAccessor)
+					{
+						UE_LOG_ERROR_UHT(TEXT("A function cannot be both BlueprintImplementableEvent and a Blueprint Property accessor!"));
+					}
+					else if ( (FuncInfo.FunctionFlags & FUNC_Private) )
+					{
+						UE_LOG_ERROR_UHT(TEXT("A Private function cannot be a BlueprintImplementableEvent!") );
+					}
+
+					FuncInfo.FunctionFlags |= FUNC_Event;
+					FuncInfo.FunctionFlags |= FUNC_BlueprintEvent;
+					FuncInfo.FunctionFlags &= ~FUNC_Native;
+				}
+				break;
+
+				case EFunctionSpecifier::Exec:
+				{
+					FuncInfo.FunctionFlags |= FUNC_Exec;
+					if( FuncInfo.FunctionFlags & FUNC_Net )
+					{
+						UE_LOG_ERROR_UHT(TEXT("Exec functions cannot be replicated!") );
+					}
+				}
+				break;
+
+				case EFunctionSpecifier::SealedEvent:
+				{
+					FuncInfo.bSealedEvent = true;
+				}
+				break;
+
+				case EFunctionSpecifier::Server:
+				{
+					if ((FuncInfo.FunctionFlags & FUNC_BlueprintEvent) != 0)
+					{
+						FError::Throwf(TEXT("BlueprintImplementableEvent or BlueprintNativeEvent functions cannot be declared as Client or Server"));
+					}
+
+					FuncInfo.FunctionFlags |= FUNC_Net;
+					FuncInfo.FunctionFlags |= FUNC_NetServer;
+
+					if (Specifier.Values.Num())
+					{
+						FuncInfo.CppImplName = Specifier.Values[0];
+					}
+
+					if( FuncInfo.FunctionFlags & FUNC_Exec )
+					{
+						UE_LOG_ERROR_UHT(TEXT("Exec functions cannot be replicated!") );
+					}
+				}
+				break;
+
+				case EFunctionSpecifier::Client:
+				{
+					if ((FuncInfo.FunctionFlags & FUNC_BlueprintEvent) != 0)
+					{
+						FError::Throwf(TEXT("BlueprintImplementableEvent or BlueprintNativeEvent functions cannot be declared as Client or Server"));
+					}
+
+					FuncInfo.FunctionFlags |= FUNC_Net;
+					FuncInfo.FunctionFlags |= FUNC_NetClient;
+
+					if (Specifier.Values.Num())
+					{
+						FuncInfo.CppImplName = Specifier.Values[0];
+					}
+				}
+				break;
+
+				case EFunctionSpecifier::NetMulticast:
+				{
+					if ((FuncInfo.FunctionFlags & FUNC_BlueprintEvent) != 0)
+					{
+						FError::Throwf(TEXT("BlueprintImplementableEvent or BlueprintNativeEvent functions cannot be declared as Multicast"));
+					}
+
+					FuncInfo.FunctionFlags |= FUNC_Net;
+					FuncInfo.FunctionFlags |= FUNC_NetMulticast;
+				}
+				break;
+
+				case EFunctionSpecifier::ServiceRequest:
+				{
+					if ((FuncInfo.FunctionFlags & FUNC_BlueprintEvent) != 0)
+					{
+						FError::Throwf(TEXT("BlueprintImplementableEvent or BlueprintNativeEvent functions cannot be declared as a ServiceRequest"));
+					}
+
+					FuncInfo.FunctionFlags |= FUNC_Net;
+					FuncInfo.FunctionFlags |= FUNC_NetReliable;
+					FuncInfo.FunctionFlags |= FUNC_NetRequest;
+					FuncInfo.FunctionExportFlags |= FUNCEXPORT_CustomThunk;
+
+					ParseNetServiceIdentifiers(FuncInfo, Specifier.Values);
+
+					if (FuncInfo.EndpointName.Len() == 0)
+					{
+						FError::Throwf(TEXT("ServiceRequest needs to specify an endpoint name"));
+					}
+				}
+				break;
+
+				case EFunctionSpecifier::ServiceResponse:
+				{
+					if ((FuncInfo.FunctionFlags & FUNC_BlueprintEvent) != 0)
+					{
+						FError::Throwf(TEXT("BlueprintImplementableEvent or BlueprintNativeEvent functions cannot be declared as a ServiceResponse"));
+					}
+
+					FuncInfo.FunctionFlags |= FUNC_Net;
+					FuncInfo.FunctionFlags |= FUNC_NetReliable;
+					FuncInfo.FunctionFlags |= FUNC_NetResponse;
+
+					ParseNetServiceIdentifiers(FuncInfo, Specifier.Values);
+
+					if (FuncInfo.EndpointName.Len() == 0)
+					{
+						FError::Throwf(TEXT("ServiceResponse needs to specify an endpoint name"));
+					}
+				}
+				break;
+
+				case EFunctionSpecifier::Reliable:
+				{
+					FuncInfo.FunctionFlags |= FUNC_NetReliable;
+				}
+				break;
+
+				case EFunctionSpecifier::Unreliable:
+				{
+					bSpecifiedUnreliable = true;
+				}
+				break;
+
+				case EFunctionSpecifier::CustomThunk:
+				{
+					FuncInfo.FunctionExportFlags |= FUNCEXPORT_CustomThunk;
+				}
+				break;
+
+				case EFunctionSpecifier::BlueprintCallable:
+				{
+					FuncInfo.FunctionFlags |= FUNC_BlueprintCallable;
+				}
+				break;
+
+				case EFunctionSpecifier::BlueprintGetter:
+				{
+					if (FuncInfo.FunctionFlags & FUNC_Event)
+					{
+						UE_LOG_ERROR_UHT(TEXT("Function cannot be a blueprint event and a blueprint getter."));
+					}
+
+					bSawPropertyAccessor = true;
+					FuncInfo.FunctionFlags |= FUNC_BlueprintCallable;
+					FuncInfo.FunctionFlags |= FUNC_BlueprintPure;
+					MetaData.Add(NAME_BlueprintGetter);
+				}
+				break;
+
+				case EFunctionSpecifier::BlueprintSetter:
+				{
+					if (FuncInfo.FunctionFlags & FUNC_Event)
+					{
+						UE_LOG_ERROR_UHT(TEXT("Function cannot be a blueprint event and a blueprint setter."));
+					}
+
+					bSawPropertyAccessor = true;
+					FuncInfo.FunctionFlags |= FUNC_BlueprintCallable;
+					MetaData.Add(NAME_BlueprintSetter);
+				}
+				break;
+
+				case EFunctionSpecifier::BlueprintPure:
+				{
+					bool bIsPure = true;
+					if (Specifier.Values.Num() == 1)
+					{
+						FString IsPureStr = Specifier.Values[0];
+						bIsPure = IsPureStr.ToBool();
+					}
+
+					// This function can be called, and is also pure.
+					FuncInfo.FunctionFlags |= FUNC_BlueprintCallable;
+
+					if (bIsPure)
+					{
+						FuncInfo.FunctionFlags |= FUNC_BlueprintPure;
+					}
+					else
+					{
+						FuncInfo.bForceBlueprintImpure = true;
+					}
+				}
+				break;
+
+				case EFunctionSpecifier::BlueprintAuthorityOnly:
+				{
+					FuncInfo.FunctionFlags |= FUNC_BlueprintAuthorityOnly;
+				}
+				break;
+
+				case EFunctionSpecifier::BlueprintCosmetic:
+				{
+					FuncInfo.FunctionFlags |= FUNC_BlueprintCosmetic;
+				}
+				break;
+
+				case EFunctionSpecifier::WithValidation:
+				{
+					FuncInfo.FunctionFlags |= FUNC_NetValidate;
+
+					if (Specifier.Values.Num())
+					{
+						FuncInfo.CppValidationImplName = Specifier.Values[0];
+					}
+				}
+				break;
+			}
+		}
+
+		if (FuncInfo.FunctionFlags & FUNC_Net)
+		{
+			// Network replicated functions are always events
+			FuncInfo.FunctionFlags |= FUNC_Event;
+
+			check(!(FuncInfo.FunctionFlags & (FUNC_BlueprintEvent | FUNC_Exec)));
+
+			bool bIsNetService  = !!(FuncInfo.FunctionFlags & (FUNC_NetRequest | FUNC_NetResponse));
+			bool bIsNetReliable = !!(FuncInfo.FunctionFlags & FUNC_NetReliable);
+
+			if (FuncInfo.FunctionFlags & FUNC_Static)
+			{
+				UE_LOG_ERROR_UHT(TEXT("Static functions can't be replicated"));
+			}
+
+			if (!bIsNetReliable && !bSpecifiedUnreliable && !bIsNetService)
+			{
+				UE_LOG_ERROR_UHT(TEXT("Replicated function: 'reliable' or 'unreliable' is required"));
+			}
+
+			if (bIsNetReliable && bSpecifiedUnreliable && !bIsNetService)
+			{
+				UE_LOG_ERROR_UHT(TEXT("'reliable' and 'unreliable' are mutually exclusive"));
+			}
+		}
+		else if (FuncInfo.FunctionFlags & FUNC_NetReliable)
+		{
+			UE_LOG_ERROR_UHT(TEXT("'reliable' specified without 'client' or 'server'"));
+		}
+		else if (bSpecifiedUnreliable)
+		{
+			UE_LOG_ERROR_UHT(TEXT("'unreliable' specified without 'client' or 'server'"));
+		}
+
+		if (FuncInfo.bSealedEvent && !(FuncInfo.FunctionFlags & FUNC_Event))
+		{
+			UE_LOG_ERROR_UHT(TEXT("SealedEvent may only be used on events"));
+		}
+
+		if (FuncInfo.bSealedEvent && FuncInfo.FunctionFlags & FUNC_BlueprintEvent)
+		{
+			UE_LOG_ERROR_UHT(TEXT("SealedEvent cannot be used on Blueprint events"));
+		}
+
+		if (FuncInfo.bForceBlueprintImpure && (FuncInfo.FunctionFlags & FUNC_BlueprintPure) != 0)
+		{
+			UE_LOG_ERROR_UHT(TEXT("BlueprintPure (or BlueprintPure=true) and BlueprintPure=false should not both appear on the same function, they are mutually exclusive"));
+		}
+	}
+
+	void AddEditInlineMetaData(TMap<FName, FString>& MetaData)
+	{
+		MetaData.Add(NAME_EditInline, TEXT("true"));
+	}
+
+	const TCHAR* GetHintText(EVariableCategory::Type VariableCategory)
+	{
+		switch (VariableCategory)
+		{
+			case EVariableCategory::ReplicatedParameter:
+			case EVariableCategory::RegularParameter:
+				return TEXT("Function parameter");
+
+			case EVariableCategory::Return:
+				return TEXT("Function return type");
+
+			case EVariableCategory::Member:
+				return TEXT("Member variable declaration");
+
+			default:
+				FError::Throwf(TEXT("Unknown variable category"));
+		}
+
+		// Unreachable
+		check(false);
+		return nullptr;
+	}
+
+	// Check to see if anything in the class hierarchy passed in has CLASS_DefaultToInstanced
+	bool DoesAnythingInHierarchyHaveDefaultToInstanced(UClass* TestClass)
+	{
+		bool bDefaultToInstanced = false;
+
+		UClass* Search = TestClass;
+		while (!bDefaultToInstanced && (Search != NULL))
+		{
+			bDefaultToInstanced = Search->HasAnyClassFlags(CLASS_DefaultToInstanced);
+			if (!bDefaultToInstanced && !Search->HasAnyClassFlags(CLASS_Intrinsic | CLASS_Parsed))
+			{
+				// The class might not have been parsed yet, look for declaration data.
+				if (FClassDeclarationMetaData* ClassDeclarationDataPtr = GClassDeclarations.Find(Search->GetFName()))
+				{
+					bDefaultToInstanced = !!(ClassDeclarationDataPtr->ClassFlags & CLASS_DefaultToInstanced);
+				}
+			}
+			Search = Search->GetSuperClass();
+		}
+
+		return bDefaultToInstanced;
+	}
+
+	FProperty* CreateVariableProperty(FPropertyBase& VarProperty, FFieldVariant Scope, const FName& Name, EObjectFlags ObjectFlags, EVariableCategory::Type VariableCategory, FUnrealSourceFile* UnrealSourceFile)
+	{
+		// Check if it's an enum class property
+		if (const EUnderlyingEnumType* EnumPropType = GEnumUnderlyingTypes.Find(VarProperty.Enum))
+		{
+			FPropertyBase UnderlyingProperty = VarProperty;
+			UnderlyingProperty.Enum = nullptr;
+			switch (*EnumPropType)
+			{
+				case EUnderlyingEnumType::int8:        UnderlyingProperty.Type = CPT_Int8;   break;
+				case EUnderlyingEnumType::int16:       UnderlyingProperty.Type = CPT_Int16;  break;
+				case EUnderlyingEnumType::int32:       UnderlyingProperty.Type = CPT_Int;    break;
+				case EUnderlyingEnumType::int64:       UnderlyingProperty.Type = CPT_Int64;  break;
+				case EUnderlyingEnumType::uint8:       UnderlyingProperty.Type = CPT_Byte;   break;
+				case EUnderlyingEnumType::uint16:      UnderlyingProperty.Type = CPT_UInt16; break;
+				case EUnderlyingEnumType::uint32:      UnderlyingProperty.Type = CPT_UInt32; break;
+				case EUnderlyingEnumType::uint64:      UnderlyingProperty.Type = CPT_UInt64; break;
+				case EUnderlyingEnumType::Unspecified: UnderlyingProperty.Type = CPT_Int;    break;
+
+				default:
+					check(false);
+			}
+
+			if (*EnumPropType == EUnderlyingEnumType::Unspecified)
+			{
+				UnderlyingProperty.IntType = EIntType::Unsized;
+			}
+
+			FEnumProperty* Result = new FEnumProperty(Scope, Name, ObjectFlags);
+			FNumericProperty* UnderlyingProp = CastFieldChecked<FNumericProperty>(CreateVariableProperty(UnderlyingProperty, Result, TEXT("UnderlyingType"), ObjectFlags, VariableCategory, UnrealSourceFile));
+			Result->UnderlyingProp = MoveTemp(UnderlyingProp);
+			Result->Enum = VarProperty.Enum;
+
+			return Result;
+		}
+
+		switch (VarProperty.Type)
+		{
+			case CPT_Byte:
+			{
+				FByteProperty* Result = new FByteProperty(Scope, Name, ObjectFlags);
+				Result->Enum = VarProperty.Enum;
+				check(VarProperty.IntType == EIntType::Sized);
+				return Result;
+			}
+
+			case CPT_Int8:
+			{
+				FInt8Property* Result = new FInt8Property(Scope, Name, ObjectFlags);
+				check(VarProperty.IntType == EIntType::Sized);
+				return Result;
+			}
+
+			case CPT_Int16:
+			{
+				FInt16Property* Result = new FInt16Property(Scope, Name, ObjectFlags);
+				check(VarProperty.IntType == EIntType::Sized);
+				return Result;
+			}
+
+			case CPT_Int:
+			{
+				FIntProperty* Result = new FIntProperty(Scope, Name, ObjectFlags);
+				if (VarProperty.IntType == EIntType::Unsized)
+				{
+					GUnsizedProperties.Add(Result);
+				}
+				return Result;
+			}
+
+			case CPT_Int64:
+			{
+				FInt64Property* Result = new FInt64Property(Scope, Name, ObjectFlags);
+				check(VarProperty.IntType == EIntType::Sized);
+				return Result;
+			}
+
+			case CPT_UInt16:
+			{
+				FUInt16Property* Result = new FUInt16Property(Scope, Name, ObjectFlags);
+				check(VarProperty.IntType == EIntType::Sized);
+				return Result;
+			}
+
+			case CPT_UInt32:
+			{
+				FUInt32Property* Result = new FUInt32Property(Scope, Name, ObjectFlags);
+				if (VarProperty.IntType == EIntType::Unsized)
+				{
+					GUnsizedProperties.Add(Result);
+				}
+				return Result;
+			}
+
+			case CPT_UInt64:
+			{
+				FUInt64Property* Result = new FUInt64Property(Scope, Name, ObjectFlags);
+				check(VarProperty.IntType == EIntType::Sized);
+				return Result;
+			}
+
+			case CPT_Bool:
+			{
+				FBoolProperty* Result = new FBoolProperty(Scope, Name, ObjectFlags);
+				Result->SetBoolSize(sizeof(bool), true);
+				return Result;
+			}
+
+			case CPT_Bool8:
+			{
+				FBoolProperty* Result = new FBoolProperty(Scope, Name, ObjectFlags);
+				Result->SetBoolSize((VariableCategory == EVariableCategory::Return) ? sizeof(bool) : sizeof(uint8), VariableCategory == EVariableCategory::Return);
+				return Result;
+			}
+
+			case CPT_Bool16:
+			{
+				FBoolProperty* Result = new FBoolProperty(Scope, Name, ObjectFlags);
+				Result->SetBoolSize((VariableCategory == EVariableCategory::Return) ? sizeof(bool) : sizeof(uint16), VariableCategory == EVariableCategory::Return);
+				return Result;
+			}
+
+			case CPT_Bool32:
+			{
+				FBoolProperty* Result = new FBoolProperty(Scope, Name, ObjectFlags);
+				Result->SetBoolSize((VariableCategory == EVariableCategory::Return) ? sizeof(bool) : sizeof(uint32), VariableCategory == EVariableCategory::Return);
+				return Result;
+			}
+
+			case CPT_Bool64:
+			{
+				FBoolProperty* Result = new FBoolProperty(Scope, Name, ObjectFlags);
+				Result->SetBoolSize((VariableCategory == EVariableCategory::Return) ? sizeof(bool) : sizeof(uint64), VariableCategory == EVariableCategory::Return);
+				return Result;
+			}
+
+			case CPT_Float:
+			{
+				FFloatProperty* Result = new FFloatProperty(Scope, Name, ObjectFlags);
+				return Result;
+			}
+
+			case CPT_Double:
+			{
+				FDoubleProperty* Result = new FDoubleProperty(Scope, Name, ObjectFlags);
+				return Result;
+			}
+
+			case CPT_ObjectReference:
+				check(VarProperty.PropertyClass);
+
+				if (VarProperty.PropertyClass->IsChildOf(UClass::StaticClass()))
+				{
+					FClassProperty* Result = new FClassProperty(Scope, Name, ObjectFlags);
+					Result->MetaClass     = VarProperty.MetaClass;
+					Result->PropertyClass = VarProperty.PropertyClass;
+					return Result;
+				}
+				else
+				{
+					if (DoesAnythingInHierarchyHaveDefaultToInstanced(VarProperty.PropertyClass))
+					{
+						VarProperty.PropertyFlags |= CPF_InstancedReference;
+						AddEditInlineMetaData(VarProperty.MetaData);
+					}
+
+					FObjectProperty* Result = new FObjectProperty(Scope, Name, ObjectFlags);
+					Result->PropertyClass = VarProperty.PropertyClass;
+					return Result;
+				}
+
+			case CPT_WeakObjectReference:
+			{
+				check(VarProperty.PropertyClass);
+
+				FWeakObjectProperty* Result = new FWeakObjectProperty(Scope, Name, ObjectFlags);
+				Result->PropertyClass = VarProperty.PropertyClass;
+				return Result;
+			}
+
+			case CPT_LazyObjectReference:
+			{
+				check(VarProperty.PropertyClass);
+
+				FLazyObjectProperty* Result = new FLazyObjectProperty(Scope, Name, ObjectFlags);
+				Result->PropertyClass = VarProperty.PropertyClass;
+				return Result;
+			}
+
+			case CPT_SoftObjectReference:
+				check(VarProperty.PropertyClass);
+
+				if (VarProperty.PropertyClass->IsChildOf(UClass::StaticClass()))
+				{
+					FSoftClassProperty* Result = new FSoftClassProperty(Scope, Name, ObjectFlags);
+					Result->MetaClass     = VarProperty.MetaClass;
+					Result->PropertyClass = VarProperty.PropertyClass;
+					return Result;
+				}
+				else
+				{
+					FSoftObjectProperty* Result = new FSoftObjectProperty(Scope, Name, ObjectFlags);
+					Result->PropertyClass = VarProperty.PropertyClass;
+					return Result;
+				}
+
+			case CPT_Interface:
+			{
+				check(VarProperty.PropertyClass);
+				check(VarProperty.PropertyClass->HasAnyClassFlags(CLASS_Interface));
+
+				FInterfaceProperty* Result = new  FInterfaceProperty(Scope, Name, ObjectFlags);
+				Result->InterfaceClass = VarProperty.PropertyClass;
+				return Result;
+			}
+
+			case CPT_Name:
+			{
+				FNameProperty* Result = new FNameProperty(Scope, Name, ObjectFlags);
+				return Result;
+			}
+
+			case CPT_String:
+			{
+				FStrProperty* Result = new FStrProperty(Scope, Name, ObjectFlags);
+				return Result;
+			}
+
+			case CPT_Text:
+			{
+				FTextProperty* Result = new FTextProperty(Scope, Name, ObjectFlags);
+				return Result;
+			}
+
+			case CPT_Struct:
+			{
+				if (VarProperty.Struct->StructFlags & STRUCT_HasInstancedReference)
+				{
+					VarProperty.PropertyFlags |= CPF_ContainsInstancedReference;
+				}
+
+				FStructProperty* Result = new FStructProperty(Scope, Name, ObjectFlags);
+				Result->Struct = VarProperty.Struct;
+				return Result;
+			}
+
+			case CPT_Delegate:
+			{
+				FDelegateProperty* Result = new FDelegateProperty(Scope, Name, ObjectFlags);
+				return Result;
+			}
+
+			case CPT_MulticastDelegate:
+			{
+				FMulticastDelegateProperty* Result;
+				if (VarProperty.Function->IsA<USparseDelegateFunction>())
+				{
+					Result = new FMulticastSparseDelegateProperty(Scope, Name, ObjectFlags);
+				}
+				else
+				{
+					Result = new FMulticastInlineDelegateProperty(Scope, Name, ObjectFlags);
+				}
+				return Result;
+			}
+
+			case CPT_FieldPath:
+			{
+				FFieldPathProperty* Result = new FFieldPathProperty(Scope, Name, ObjectFlags);
+				Result->PropertyClass = VarProperty.PropertyPathClass;
+
+				return Result;
+			}
+
+			default:
+				FError::Throwf(TEXT("Unknown property type %i"), (uint8)VarProperty.Type);
+		}
+
+		// Unreachable
+		check(false); //-V779
+		return nullptr;
+	}
+
+	/**
+	 * Ensures at script compile time that the metadata formatting is correct
+	 * @param	InKey			the metadata key being added
+	 * @param	InValue			the value string that will be associated with the InKey
+	 */
+	void ValidateMetaDataFormat(FFieldVariant Field, const FName InKey, const FString& InValue)
+	{
+		switch (GetCheckedMetadataSpecifier(InKey))
+		{
+			default:
+			{
+				// Don't need to validate this specifier
+			}
+			break;
+
+			case ECheckedMetadataSpecifier::UIMin:
+			case ECheckedMetadataSpecifier::UIMax:
+			case ECheckedMetadataSpecifier::ClampMin:
+			case ECheckedMetadataSpecifier::ClampMax:
+			{
+				if (!InValue.IsNumeric())
+				{
+					FError::Throwf(TEXT("Metadata value for '%s' is non-numeric : '%s'"), *InKey.ToString(), *InValue);
+				}
+			}
+			break;
+
+			case ECheckedMetadataSpecifier::BlueprintProtected:
+			{
+				if (Field.IsUObject())
+				{
+					UFunction* Function = Field.Get<UFunction>();
+					if (Function->HasAnyFunctionFlags(FUNC_Static))
+					{
+						// Determine if it's a function library
+						UClass* Class = Function->GetOuterUClass();
+						while (Class != nullptr && Class->GetSuperClass() != UObject::StaticClass())
+						{
+							Class = Class->GetSuperClass();
+						}
+
+						if (Class != nullptr && Class->GetName() == TEXT("BlueprintFunctionLibrary"))
+						{
+							FError::Throwf(TEXT("%s doesn't make sense on static method '%s' in a blueprint function library"), *InKey.ToString(), *Function->GetName());
+						}
+					}
+				}
+			}
+			break;
+
+			case ECheckedMetadataSpecifier::CommutativeAssociativeBinaryOperator:
+			{
+				if (UFunction* Function = Field.Get<UFunction>())
+				{
+					bool bGoodParams = (Function->NumParms == 3);
+					if (bGoodParams)
+					{
+						FProperty* FirstParam = nullptr;
+						FProperty* SecondParam = nullptr;
+						FProperty* ReturnValue = nullptr;
+
+						TFieldIterator<FProperty> It(Function);
+
+						auto GetNextParam = [&]()
+						{
+							if (It)
+							{
+								if (It->HasAnyPropertyFlags(CPF_ReturnParm))
+								{
+									ReturnValue = *It;
+								}
+								else
+								{
+									if (FirstParam == nullptr)
+									{
+										FirstParam = *It;
+									}
+									else if (SecondParam == nullptr)
+									{
+										SecondParam = *It;
+									}
+								}
+								++It;
+							}
+						};
+
+						GetNextParam();
+						GetNextParam();
+						GetNextParam();
+						ensure(!It);
+
+						if (ReturnValue == nullptr || SecondParam == nullptr || !SecondParam->SameType(FirstParam))
+						{
+							bGoodParams = false;
+						}
+					}
+
+					if (!bGoodParams)
+					{
+						UE_LOG_ERROR_UHT(TEXT("Commutative asssociative binary operators must have exactly 2 parameters of the same type and a return value."));
+					}
+				}
+			}
+			break;
+
+			case ECheckedMetadataSpecifier::ExpandBoolAsExecs:
+			case ECheckedMetadataSpecifier::ExpandEnumAsExecs:
+			{
+				if (UFunction* Function = Field.Get<UFunction>())
+				{
+					// multiple entry parsing in the same format as eg SetParam.
+					TArray<FString> RawGroupings;
+					InValue.ParseIntoArray(RawGroupings, TEXT(","), false);
+
+					FProperty* FirstInput = nullptr;
+					for (const FString& RawGroup : RawGroupings)
+					{
+						TArray<FString> IndividualEntries;
+						RawGroup.ParseIntoArray(IndividualEntries, TEXT("|"));
+
+						for (const FString& Entry : IndividualEntries)
+						{
+							if (Entry.IsEmpty())
+							{
+								continue;
+							}
+							
+							FField* FoundField = FHeaderParser::FindProperty(Function, *Entry, false);
+							if (!FoundField)
+							{
+								UE_LOG_ERROR_UHT(TEXT("Function does not have a parameter named '%s'"), *Entry);
+							}
+							else if (FProperty* Prop = CastField<FProperty>(FoundField))
+							{
+								if (!Prop->HasAnyPropertyFlags(CPF_ReturnParm) &&
+
+								    (!Prop->HasAnyPropertyFlags(CPF_OutParm) ||
+									Prop->HasAnyPropertyFlags(CPF_ReferenceParm)))
+								{
+									if (!FirstInput)
+									{
+										FirstInput = Prop;
+									}
+									else
+									{
+										UE_LOG_ERROR_UHT(TEXT("Function already specified an ExpandEnumAsExec input (%s), but '%s' is also an input parameter. Only one is permitted."), *FirstInput->GetName(), *Entry);
+									}
+								}
+							}
+						}
+					}
+				}
+			}
+			break;
+
+			case ECheckedMetadataSpecifier::DevelopmentStatus:
+			{
+				const FString EarlyAccessValue(TEXT("EarlyAccess"));
+				const FString ExperimentalValue(TEXT("Experimental"));
+				if ((InValue != EarlyAccessValue) && (InValue != ExperimentalValue))
+				{
+					FError::Throwf(TEXT("'%s' metadata was '%s' but it must be %s or %s"), *InKey.ToString(), *InValue, *ExperimentalValue, *EarlyAccessValue);
+				}
+			}
+			break;
+
+			case ECheckedMetadataSpecifier::Units:
+			{
+				// Check for numeric property
+				FField* MaybeProperty = Field.ToField();
+				if (!MaybeProperty->IsA<FNumericProperty>() && !MaybeProperty->IsA<FStructProperty>())
+				{
+					FError::Throwf(TEXT("'Units' meta data can only be applied to numeric and struct properties"));
+				}
+
+				if (!FUnitConversion::UnitFromString(*InValue))
+				{
+					FError::Throwf(TEXT("Unrecognized units (%s) specified for property '%s'"), *InValue, *Field.GetFullName());
+				}
+			}
+			break;
+
+			case ECheckedMetadataSpecifier::DocumentationPolicy:
+			{
+				const TCHAR* StrictValue = TEXT("Strict");
+				if (InValue != StrictValue)
+				{
+					FError::Throwf(TEXT("'%s' metadata was '%s' but it must be %s"), *InKey.ToString(), *InValue, *StrictValue);
+				}
+			}
+			break;
+		}
+	}
+
+	// Ensures at script compile time that the metadata formatting is correct
+	void ValidateMetaDataFormat(FFieldVariant Field, const TMap<FName, FString>& MetaData)
+	{
+		for (const TPair<FName, FString>& Pair : MetaData)
+		{
+			ValidateMetaDataFormat(Field, Pair.Key, Pair.Value);
+		}
+	}
+
+	// Validates the metadata, then adds it to the class data
+	void AddMetaDataToClassData(FFieldVariant Field, TMap<FName, FString>&& InMetaData)
+	{
+		// Evaluate any key redirects on the passed in pairs
+		for (TPair<FName, FString>& Pair : InMetaData)
+		{
+			FName& CurrentKey = Pair.Key;
+			FName NewKey = UMetaData::GetRemappedKeyName(CurrentKey);
+
+			if (NewKey != NAME_None)
+			{
+				UE_LOG_WARNING_UHT(TEXT("Remapping old metadata key '%s' to new key '%s', please update the declaration."), *CurrentKey.ToString(), *NewKey.ToString());
+				CurrentKey = NewKey;
+			}
+		}
+
+		// Finish validating and associate the metadata with the field
+		ValidateMetaDataFormat(Field, InMetaData);
+		if (Field.IsUObject())
+		{
+			FClassMetaData::AddMetaData(CastChecked<UField>(Field.ToUObject()), MoveTemp(InMetaData));
+		}
+		else
+		{
+			FClassMetaData::AddMetaData(Field.ToField(), MoveTemp(InMetaData));
+		}
+	}
+
+	bool IsPropertySupportedByBlueprint(const FProperty* Property, bool bMemberVariable)
+	{
+		if (Property == nullptr)
+		{
+			return false;
+		}
+		if (const FArrayProperty* ArrayProperty = CastField<const FArrayProperty>(Property))
+		{
+			// Script VM doesn't support array of weak ptrs.
+			return IsPropertySupportedByBlueprint(ArrayProperty->Inner, false);
+		}
+		else if (const FSetProperty* SetProperty = CastField<const FSetProperty>(Property))
+		{
+			return IsPropertySupportedByBlueprint(SetProperty->ElementProp, false);
+		}
+		else if (const FMapProperty* MapProperty = CastField<const FMapProperty>(Property))
+		{
+			return IsPropertySupportedByBlueprint(MapProperty->KeyProp, false) &&
+				IsPropertySupportedByBlueprint(MapProperty->ValueProp, false);
+		}
+		else if (const FStructProperty* StructProperty = CastField<const FStructProperty>(Property))
+		{
+			return (StructProperty->Struct->GetBoolMetaDataHierarchical(NAME_BlueprintType));
+		}
+
+		const bool bSupportedType = Property->IsA<FInterfaceProperty>()
+			|| Property->IsA<FClassProperty>()
+			|| Property->IsA<FSoftObjectProperty>()
+			|| Property->IsA<FObjectProperty>()
+			|| Property->IsA<FFloatProperty>()
+			|| Property->IsA<FDoubleProperty>()
+			|| Property->IsA<FIntProperty>()
+			|| Property->IsA<FInt64Property>()
+			|| Property->IsA<FByteProperty>()
+			|| Property->IsA<FNameProperty>()
+			|| Property->IsA<FBoolProperty>()
+			|| Property->IsA<FStrProperty>()
+			|| Property->IsA<FTextProperty>()
+			|| Property->IsA<FDelegateProperty>()
+			|| Property->IsA<FEnumProperty>()
+			|| Property->IsA<FFieldPathProperty>();
+
+		const bool bIsSupportedMemberVariable = Property->IsA<FWeakObjectProperty>() || Property->IsA<FMulticastDelegateProperty>();
+
+		return bSupportedType || (bIsSupportedMemberVariable && bMemberVariable);
+	}
+
+	void SkipAlignasIfNecessary(FBaseParser& Parser)
+	{
+		if (Parser.MatchIdentifier(TEXT("alignas"), ESearchCase::CaseSensitive))
+		{
+			Parser.RequireSymbol(TEXT('('), TEXT("'alignas'"));
+			Parser.RequireAnyConstInt(TEXT("'alignas'"));
+			Parser.RequireSymbol(TEXT(')'), TEXT("'alignas'"));
+		}
+	}
+
+	void SkipDeprecatedMacroIfNecessary(FBaseParser& Parser)
+	{
+		FToken MacroToken;
+		if (!Parser.GetToken(MacroToken))
+		{
+			return;
+		}
+
+		if (MacroToken.TokenType != TOKEN_Identifier || (FCString::Strcmp(MacroToken.Identifier, TEXT("DEPRECATED")) != 0 && FCString::Strcmp(MacroToken.Identifier, TEXT("UE_DEPRECATED")) != 0))
+		{
+			Parser.UngetToken(MacroToken);
+			return;
+		}
+
+		auto ErrorMessageGetter = [&MacroToken]() { return FString::Printf(TEXT("%s macro"), MacroToken.Identifier); };
+
+		Parser.RequireSymbol(TEXT('('), ErrorMessageGetter);
+
+		FToken Token;
+		if (Parser.GetToken(Token) && (Token.Type != CPT_Float || Token.TokenType != TOKEN_Const))
+		{
+			FError::Throwf(TEXT("Expected engine version in %s macro"), MacroToken.Identifier);
+		}
+
+		Parser.RequireSymbol(TEXT(','), ErrorMessageGetter);
+		if (Parser.GetToken(Token) && (Token.Type != CPT_String || Token.TokenType != TOKEN_Const))
+		{
+			FError::Throwf(TEXT("Expected deprecation message in %s macro"), MacroToken.Identifier);
+		}
+
+		Parser.RequireSymbol(TEXT(')'), ErrorMessageGetter);
+	}
+
+	static const TCHAR* GLayoutMacroNames[] = {
+		TEXT("LAYOUT_ARRAY"),
+		TEXT("LAYOUT_ARRAY_EDITORONLY"),
+		TEXT("LAYOUT_BITFIELD"),
+		TEXT("LAYOUT_BITFIELD_EDITORONLY"),
+		TEXT("LAYOUT_FIELD"),
+		TEXT("LAYOUT_FIELD_EDITORONLY"),
+		TEXT("LAYOUT_FIELD_INITIALIZED"),
+	};
+}
+
+/////////////////////////////////////////////////////
+// FScriptLocation
+
+FHeaderParser* FScriptLocation::Compiler = NULL;
+
+FScriptLocation::FScriptLocation()
+{
+	if ( Compiler != NULL )
+	{
+		Compiler->InitScriptLocation(*this);
+	}
+}
+
+/////////////////////////////////////////////////////
+// FHeaderParser
+
+FString FHeaderParser::GetContext()
+{
+	FFileScope* FileScope = GetCurrentFileScope();
+	FUnrealSourceFile* SourceFile = FileScope ? FileScope->GetSourceFile() : GetCurrentSourceFile();
+	FString ScopeFilename = SourceFile
+		? IFileManager::Get().ConvertToAbsolutePathForExternalAppForRead(*SourceFile->GetFilename())
+		: TEXT("UNKNOWN");
+
+	return FString::Printf(TEXT("%s(%i)"), *ScopeFilename, InputLine);
+}
+
+/*-----------------------------------------------------------------------------
+	Code emitting.
+-----------------------------------------------------------------------------*/
+
+
+//
+// Get a qualified class.
+//
+FClass* FHeaderParser::GetQualifiedClass(const FClasses& AllClasses, const TCHAR* Thing)
+{
+	TCHAR ClassName[256]=TEXT("");
+
+	FToken Token;
+	if (GetIdentifier(Token))
+	{
+		RedirectTypeIdentifier(Token);
+
+		FCString::Strncat( ClassName, Token.Identifier, UE_ARRAY_COUNT(ClassName) );
+	}
+
+	if (!ClassName[0])
+	{
+		FError::Throwf(TEXT("%s: Missing class name"), Thing );
+	}
+
+	return AllClasses.FindScriptClassOrThrow(ClassName);
+}
+
+/*-----------------------------------------------------------------------------
+	Fields.
+-----------------------------------------------------------------------------*/
+
+/**
+ * Find a field in the specified context.  Starts with the specified scope, then iterates
+ * through the Outer chain until the field is found.
+ * 
+ * @param	InScope				scope to start searching for the field in 
+ * @param	InIdentifier		name of the field we're searching for
+ * @param	bIncludeParents		whether to allow searching in the scope of a parent struct
+ * @param	FieldClass			class of the field to search for.  used to e.g. search for functions only
+ * @param	Thing				hint text that will be used in the error message if an error is encountered
+ *
+ * @return	a pointer to a UField with a name matching InIdentifier, or NULL if it wasn't found
+ */
+UField* FHeaderParser::FindField
+(
+	UStruct*		Scope,
+	const TCHAR*	InIdentifier,
+	bool			bIncludeParents,
+	UClass*			FieldClass,
+	const TCHAR*	Thing
+)
+{
+	check(InIdentifier);
+	FName InName(InIdentifier, FNAME_Find);
+	if (InName != NAME_None)
+	{
+		for( ; Scope; Scope = Cast<UStruct>(Scope->GetOuter()) )
+		{
+			for (TFieldIterator<FField> It(Scope); It; ++It)
+			{
+				if (It->GetFName() == InName)
+				{
+					if (Thing)
+					{
+						FError::Throwf(TEXT("%s: expecting @todo: FProp, got %s"), Thing, /**FieldClass->GetName(),*/ *It->GetClass()->GetName());
+					}
+					return nullptr;
+				}
+			}
+			for( TFieldIterator<UField> It(Scope); It; ++It )
+			{
+				if (It->GetFName() == InName)
+				{
+					if (!It->IsA(FieldClass))
+					{
+						if (Thing)
+						{
+							FError::Throwf(TEXT("%s: expecting %s, got %s"), Thing, *FieldClass->GetName(), *It->GetClass()->GetName() );
+						}
+						return nullptr;
+					}
+					return *It;
+				}
+			}
+
+			if (!bIncludeParents)
+			{
+				break;
+			}
+		}
+	}
+
+						return NULL;
+}
+
+FField* FHeaderParser::FindProperty(UStruct* Scope, const TCHAR* InIdentifier, bool bIncludeParents, FFieldClass* FieldClass, const TCHAR* Thing)
+{
+	check(InIdentifier);
+	FName InName(InIdentifier, FNAME_Find);
+	if (InName != NAME_None)
+	{
+		for (; Scope; Scope = Cast<UStruct>(Scope->GetOuter()))
+		{
+			for (TFieldIterator<UField> It(Scope); It; ++It)
+			{
+				if (It->GetFName() == InName)
+				{
+					if (Thing)
+					{
+						FError::Throwf(TEXT("%s: expecting a property, got %s"), Thing, *It->GetClass()->GetName());
+					}
+					return nullptr;
+				}
+			}
+			for (TFieldIterator<FField> It(Scope); It; ++It)
+			{
+				if (It->GetFName() == InName)
+				{
+					if (!It->IsA(FieldClass))
+					{
+						if (Thing)
+						{
+							FError::Throwf(TEXT("%s: expecting %s: FProp, got %s"), Thing, *FieldClass->GetName(), *It->GetClass()->GetName());
+						}
+						return nullptr;
+					}
+					return *It;
+				}
+			}
+
+			if (!bIncludeParents)
+			{
+				break;
+			}
+		}
+	}
+
+	return NULL;
+}
+/**
+ * @return	true if Scope has FProperty objects in its list of fields
+ */
+bool FHeaderParser::HasMemberProperties( const UStruct* Scope )
+{
+	// it's safe to pass a NULL Scope to TFieldIterator, but this function shouldn't be called with a NULL Scope
+	checkSlow(Scope);
+	TFieldIterator<FProperty> It(Scope,EFieldIteratorFlags::ExcludeSuper);
+	return It ? true : false;
+}
+
+/**
+ * Get the parent struct specified.
+ *
+ * @param	CurrentScope	scope to start in
+ * @param	SearchName		parent scope to search for
+ *
+ * @return	a pointer to the parent struct with the specified name, or NULL if the parent couldn't be found
+ */
+UStruct* FHeaderParser::GetSuperScope( UStruct* CurrentScope, const FName& SearchName )
+{
+	UStruct* SuperScope = CurrentScope;
+	while (SuperScope && !SuperScope->GetInheritanceSuper())
+	{
+		SuperScope = CastChecked<UStruct>(SuperScope->GetOuter());
+	}
+	if (SuperScope != NULL)
+	{
+		// iterate up the inheritance chain looking for one that has the desired name
+		do
+		{
+			UStruct* NextScope = SuperScope->GetInheritanceSuper();
+			if (NextScope)
+			{
+				SuperScope = NextScope;
+			}
+			else
+			{
+				// otherwise we've failed
+				SuperScope = NULL;
+			}
+		} while (SuperScope != NULL && SuperScope->GetFName() != SearchName);
+	}
+
+	return SuperScope;
+}
+
+/**
+ * Adds source file's include path to given metadata.
+ *
+ * @param Type Type for which to add include path.
+ * @param MetaData Meta data to fill the information.
+ */
+void AddIncludePathToMetadata(UField* Type, TMap<FName, FString> &MetaData)
+{
+	// Add metadata for the include path.
+	TSharedRef<FUnrealTypeDefinitionInfo>* TypeDefinitionPtr = GTypeDefinitionInfoMap.Find(Type);
+	if (TypeDefinitionPtr != nullptr)
+	{
+		MetaData.Add(NAME_IncludePath, (*TypeDefinitionPtr)->GetUnrealSourceFile().GetIncludePath());
+	}
+}
+
+/**
+ * Adds module's relative path from given file.
+ *
+ * @param SourceFile Given source file.
+ * @param MetaData Meta data to fill the information.
+ */
+void AddModuleRelativePathToMetadata(FUnrealSourceFile& SourceFile, TMap<FName, FString> &MetaData)
+{
+	MetaData.Add(NAME_ModuleRelativePath, SourceFile.GetModuleRelativePath());
+}
+
+/**
+ * Adds module's relative path to given metadata.
+ *
+ * @param Type Type for which to add module's relative path.
+ * @param MetaData Meta data to fill the information.
+ */
+void AddModuleRelativePathToMetadata(UField* Type, TMap<FName, FString> &MetaData)
+{
+	// Add metadata for the module relative path.
+	TSharedRef<FUnrealTypeDefinitionInfo>* TypeDefinitionPtr = GTypeDefinitionInfoMap.Find(Type);
+	if (TypeDefinitionPtr != nullptr)
+	{
+		MetaData.Add(NAME_ModuleRelativePath, (*TypeDefinitionPtr)->GetUnrealSourceFile().GetModuleRelativePath());
+	}
+}
+
+/*-----------------------------------------------------------------------------
+	Variables.
+-----------------------------------------------------------------------------*/
+
+//
+// Compile an enumeration definition.
+//
+UEnum* FHeaderParser::CompileEnum()
+{
+	FUnrealSourceFile* CurrentSrcFile = GetCurrentSourceFile();
+	TSharedPtr<FFileScope> Scope = CurrentSrcFile->GetScope();
+
+	CheckAllow( TEXT("'Enum'"), ENestAllowFlags::TypeDecl );
+
+	// Get the enum specifier list
+	FToken                     EnumToken;
+	TArray<FPropertySpecifier> SpecifiersFound;
+	ReadSpecifierSetInsideMacro(SpecifiersFound, TEXT("Enum"), EnumToken.MetaData);
+
+	// We don't handle any non-metadata enum specifiers at the moment
+	if (SpecifiersFound.Num() != 0)
+	{
+		FError::Throwf(TEXT("Unknown enum specifier '%s'"), *SpecifiersFound[0].Key);
+	}
+
+	FScriptLocation DeclarationPosition;
+
+	// Check enum type. This can be global 'enum', 'namespace' or 'enum class' enums.
+	bool            bReadEnumName = false;
+	UEnum::ECppForm CppForm       = UEnum::ECppForm::Regular;
+	if (!GetIdentifier(EnumToken))
+	{
+		FError::Throwf(TEXT("Missing identifier after UENUM()") );
+	}
+
+	if (EnumToken.Matches(TEXT("namespace"), ESearchCase::CaseSensitive))
+	{
+		CppForm      = UEnum::ECppForm::Namespaced;
+		bReadEnumName = GetIdentifier(EnumToken);
+	}
+	else if (EnumToken.Matches(TEXT("enum"), ESearchCase::CaseSensitive))
+	{
+		SkipAlignasIfNecessary(*this);
+
+		if (!GetIdentifier(EnumToken))
+		{
+			FError::Throwf(TEXT("Missing identifier after enum") );
+		}
+
+		if (EnumToken.Matches(TEXT("class"), ESearchCase::CaseSensitive) || EnumToken.Matches(TEXT("struct"), ESearchCase::CaseSensitive))
+		{
+			// You can't actually have an alignas() before the class/struct keyword, but this
+			// makes the parsing easier and illegal syntax will be caught by the compiler anyway.
+			SkipAlignasIfNecessary(*this);
+
+			CppForm       = UEnum::ECppForm::EnumClass;
+			bReadEnumName = GetIdentifier(EnumToken);
+		}
+		else
+		{
+			CppForm       = UEnum::ECppForm::Regular;
+			bReadEnumName = true;
+		}
+	}
+	else
+	{
+		FError::Throwf(TEXT("UENUM() should be followed by \'enum\' or \'namespace\' keywords.") );
+	}
+
+	// Get enumeration name.
+	if (!bReadEnumName)
+	{
+		FError::Throwf(TEXT("Missing enumeration name") );
+	}
+
+	// Verify that the enumeration definition is unique within this scope.
+	UField* Existing = Scope->FindTypeByName(EnumToken.Identifier);
+	if (Existing)
+	{
+		FError::Throwf(TEXT("enum: '%s' already defined here"), *EnumToken.GetTokenName().ToString());
+	}
+
+	// Check if the enum name is using a reserved keyword
+	if (FHeaderParser::IsReservedTypeName(EnumToken))
+	{
+		FError::Throwf(TEXT("enum: '%s' uses a reserved type name."), *EnumToken.GetTokenName().ToString());
+	}
+
+	ParseFieldMetaData(EnumToken.MetaData, EnumToken.Identifier);
+	// Create enum definition.
+	UEnum* Enum = new(EC_InternalUseOnlyConstructor, CurrentSrcFile->GetPackage(), EnumToken.Identifier, RF_Public) UEnum(FObjectInitializer());
+	Scope->AddType(Enum);
+
+	if (CompilerDirectiveStack.Num() > 0 && (CompilerDirectiveStack.Last() & ECompilerDirective::WithEditorOnlyData) != 0)
+	{
+		GEditorOnlyDataTypes.Add(Enum);
+	}
+
+	GTypeDefinitionInfoMap.Add(Enum, MakeShared<FUnrealTypeDefinitionInfo>(*CurrentSrcFile, InputLine));
+
+	// Validate the metadata for the enum
+	ValidateMetaDataFormat(Enum, EnumToken.MetaData);
+
+	// Read base for enum class
+	EUnderlyingEnumType UnderlyingType = EUnderlyingEnumType::uint8;
+	if (CppForm == UEnum::ECppForm::EnumClass)
+	{
+		if (MatchSymbol(TEXT(':')))
+		{
+			FToken BaseToken;
+			if (!GetIdentifier(BaseToken))
+			{
+				FError::Throwf(TEXT("Missing enum base") );
+			}
+
+			if (!FCString::Strcmp(BaseToken.Identifier, TEXT("uint8")))
+			{
+				UnderlyingType = EUnderlyingEnumType::uint8;
+			}
+			else if (!FCString::Strcmp(BaseToken.Identifier, TEXT("uint16")))
+			{
+				UnderlyingType = EUnderlyingEnumType::uint16;
+			}
+			else if (!FCString::Strcmp(BaseToken.Identifier, TEXT("uint32")))
+			{
+				UnderlyingType = EUnderlyingEnumType::uint32;
+			}
+			else if (!FCString::Strcmp(BaseToken.Identifier, TEXT("uint64")))
+			{
+				UnderlyingType = EUnderlyingEnumType::uint64;
+			}
+			else if (!FCString::Strcmp(BaseToken.Identifier, TEXT("int8")))
+			{
+				UnderlyingType = EUnderlyingEnumType::int8;
+			}
+			else if (!FCString::Strcmp(BaseToken.Identifier, TEXT("int16")))
+			{
+				UnderlyingType = EUnderlyingEnumType::int16;
+			}
+			else if (!FCString::Strcmp(BaseToken.Identifier, TEXT("int32")))
+			{
+				UnderlyingType = EUnderlyingEnumType::int32;
+			}
+			else if (!FCString::Strcmp(BaseToken.Identifier, TEXT("int64")))
+			{
+				UnderlyingType = EUnderlyingEnumType::int64;
+			}
+			else
+			{
+				FError::Throwf(TEXT("Unsupported enum class base type: %s"), BaseToken.Identifier);
+			}
+		}
+		else
+		{
+			UnderlyingType = EUnderlyingEnumType::Unspecified;
+		}
+
+		GEnumUnderlyingTypes.Add(Enum, UnderlyingType);
+	}
+
+	if (UnderlyingType != EUnderlyingEnumType::uint8 && EnumToken.MetaData.Contains(NAME_BlueprintType))
+	{
+		FError::Throwf(TEXT("Invalid BlueprintType enum base - currently only uint8 supported"));
+	}
+
+	// Get opening brace.
+	RequireSymbol( TEXT('{'), TEXT("'Enum'") );
+
+	switch (CppForm)
+	{
+		case UEnum::ECppForm::Namespaced:
+		{
+			// Now handle the inner true enum portion
+			RequireIdentifier(TEXT("enum"), ESearchCase::CaseSensitive, TEXT("'Enum'"));
+
+			SkipAlignasIfNecessary(*this);
+
+			FToken InnerEnumToken;
+			if (!GetIdentifier(InnerEnumToken))
+			{
+				FError::Throwf(TEXT("Missing enumeration name") );
+			}
+
+			Enum->CppType = FString::Printf(TEXT("%s::%s"), EnumToken.Identifier, InnerEnumToken.Identifier);
+
+			RequireSymbol( TEXT('{'), TEXT("'Enum'") );
+		}
+		break;
+
+		case UEnum::ECppForm::Regular:
+		case UEnum::ECppForm::EnumClass:
+		{
+			Enum->CppType = EnumToken.Identifier;
+		}
+		break;
+	}
+
+	// List of all metadata generated for this enum
+	TMap<FName,FString> EnumValueMetaData = EnumToken.MetaData;
+
+	AddModuleRelativePathToMetadata(Enum, EnumValueMetaData);
+	AddFormattedPrevCommentAsTooltipMetaData(EnumValueMetaData);
+
+	// Parse all enums tags.
+	FToken TagToken;
+	TArray<TMap<FName, FString>> EntryMetaData;
+
+	TArray<TPair<FName, int64>> EnumNames;
+	int64 CurrentEnumValue = 0;
+	while (GetIdentifier(TagToken))
+	{
+		AddFormattedPrevCommentAsTooltipMetaData(TagToken.MetaData);
+
+		// Try to read an optional explicit enum value specification
+		if (MatchSymbol(TEXT('=')))
+		{
+			FToken InitToken;
+			if (!GetToken(InitToken))
+			{
+				FError::Throwf(TEXT("UENUM: missing enumerator initializer"));
+			}
+
+			int64 NewEnumValue = -1;
+			if (!InitToken.GetConstInt64(NewEnumValue))
+			{
+				// We didn't parse a literal, so set an invalid value
+				NewEnumValue = -1;
+			}
+
+			// Skip tokens until we encounter a comma, a closing brace or a UMETA declaration
+			for (;;)
+			{
+				if (!GetToken(InitToken))
+				{
+					FError::Throwf(TEXT("Enumerator: end of file encountered while parsing the initializer"));
+				}
+
+				if (InitToken.TokenType == TOKEN_Symbol)
+				{
+					if (InitToken.Matches(TEXT(',')) || InitToken.Matches(TEXT('}')))
+					{
+						UngetToken(InitToken);
+						break;
+					}
+				}
+				else if (InitToken.TokenType == TOKEN_Identifier)
+				{
+					if (FCString::Stricmp(InitToken.Identifier, TEXT("UMETA")) == 0)
+					{
+						UngetToken(InitToken);
+						break;
+					}
+				}
+
+				// There are tokens after the initializer so it's not a standalone literal,
+				// so set it to an invalid value.
+				NewEnumValue = -1;
+			}
+
+			CurrentEnumValue = NewEnumValue;
+		}
+
+		FName NewTag;
+		switch (CppForm)
+		{
+			case UEnum::ECppForm::Namespaced:
+			case UEnum::ECppForm::EnumClass:
+			{
+				NewTag = FName(*FString::Printf(TEXT("%s::%s"), EnumToken.Identifier, TagToken.Identifier), FNAME_Add);
+			}
+			break;
+
+			case UEnum::ECppForm::Regular:
+			{
+				NewTag = FName(TagToken.Identifier, FNAME_Add);
+			}
+			break;
+		}
+
+		// Save the new tag
+		EnumNames.Emplace(NewTag, CurrentEnumValue);
+
+		// Autoincrement the current enumeration value
+		if (CurrentEnumValue != -1)
+		{
+			++CurrentEnumValue;
+		}
+
+		TagToken.MetaData.Add(NAME_Name, NewTag.ToString());
+		EntryMetaData.Add(TagToken.MetaData);
+
+		// check for metadata on this enum value
+		ParseFieldMetaData(TagToken.MetaData, TagToken.Identifier);
+		if (TagToken.MetaData.Num() > 0)
+		{
+			// special case for enum value metadata - we need to prepend the key name with the enum value name
+			const FString TokenString = TagToken.Identifier;
+			for (const auto& MetaData : TagToken.MetaData)
+			{
+				FString KeyString = TokenString + TEXT(".") + MetaData.Key.ToString();
+				EnumValueMetaData.Emplace(*KeyString, MetaData.Value);
+			}
+
+			// now clear the metadata because we're going to reuse this token for parsing the next enum value
+			TagToken.MetaData.Empty();
+		}
+
+		if (!MatchSymbol(TEXT(',')))
+		{
+			FToken ClosingBrace;
+			if (!GetToken(ClosingBrace))
+			{
+				FError::Throwf(TEXT("UENUM: end of file encountered"));
+			}
+
+			if (ClosingBrace.TokenType == TOKEN_Symbol && ClosingBrace.Matches(TEXT('}')))
+			{
+				UngetToken(ClosingBrace);
+				break;
+			}
+		}
+	}
+
+	// Trailing brace and semicolon for the enum
+	RequireSymbol( TEXT('}'), TEXT("'Enum'") );
+	MatchSemi();
+
+	if (CppForm == UEnum::ECppForm::Namespaced)
+	{
+		// Trailing brace for the namespace.
+		RequireSymbol( TEXT('}'), TEXT("'Enum'") );
+	}
+
+	// Register the list of enum names.
+	if (!Enum->SetEnums(EnumNames, CppForm, false))
+	{
+		const FName MaxEnumItem      = *(Enum->GenerateEnumPrefix() + TEXT("_MAX"));
+		const int32 MaxEnumItemIndex = Enum->GetIndexByName(MaxEnumItem);
+		if (MaxEnumItemIndex != INDEX_NONE)
+		{
+			FError::Throwf(TEXT("Illegal enumeration tag specified.  Conflicts with auto-generated tag '%s'"), *MaxEnumItem.ToString());
+		}
+
+		FError::Throwf(TEXT("Unable to generate enum MAX entry '%s' due to name collision"), *MaxEnumItem.ToString());
+	}
+
+	CheckDocumentationPolicyForEnum(Enum, EnumValueMetaData, EntryMetaData);
+
+	// Add the metadata gathered for the enum to the package
+	if (EnumValueMetaData.Num() > 0)
+	{
+		UMetaData* PackageMetaData = Enum->GetOutermost()->GetMetaData();
+		checkSlow(PackageMetaData);
+
+		PackageMetaData->SetObjectValues(Enum, MoveTemp(EnumValueMetaData));
+	}
+
+	if (!Enum->IsValidEnumValue(0) && EnumToken.MetaData.Contains(NAME_BlueprintType))
+	{
+		UE_LOG_WARNING_UHT(TEXT("'%s' does not have a 0 entry! (This is a problem when the enum is initalized by default)"), *Enum->GetName());
+	}
+
+	return Enum;
+}
+
+/**
+ * Checks if a string is made up of all the same character.
+ *
+ * @param  Str The string to check for all
+ * @param  Ch  The character to check for
+ *
+ * @return True if the string is made up only of Ch characters.
+ */
+bool IsAllSameChar(const TCHAR* Str, TCHAR Ch)
+{
+	check(Str);
+
+	while (TCHAR StrCh = *Str++)
+	{
+		if (StrCh != Ch)
+		{
+			return false;
+		}
+	}
+
+	return true;
+}
+
+/**
+ * @param		Input		An input string, expected to be a script comment.
+ * @return					The input string, reformatted in such a way as to be appropriate for use as a tooltip.
+ */
+FString FHeaderParser::FormatCommentForToolTip(const FString& Input)
+{
+	// Return an empty string if there are no alpha-numeric characters or a Unicode characters above 0xFF
+	// (which would be the case for pure CJK comments) in the input string.
+	bool bFoundAlphaNumericChar = false;
+	for ( int32 i = 0 ; i < Input.Len() ; ++i )
+	{
+		if ( FChar::IsAlnum(Input[i]) || (Input[i] > 0xFF) )
+		{
+			bFoundAlphaNumericChar = true;
+			break;
+		}
+	}
+
+	if ( !bFoundAlphaNumericChar )
+	{
+		return FString();
+	}
+
+	FString Result(Input);
+
+	// Sweep out comments marked to be ignored.
+	{
+		int32 CommentStart, CommentEnd;
+		// Block comments go first
+		for (CommentStart = Result.Find(TEXT("/*~"), ESearchCase::CaseSensitive); CommentStart != INDEX_NONE; CommentStart = Result.Find(TEXT("/*~"), ESearchCase::CaseSensitive))
+		{
+			CommentEnd = Result.Find(TEXT("*/"), ESearchCase::CaseSensitive, ESearchDir::FromStart, CommentStart);
+			if (CommentEnd != INDEX_NONE)
+			{
+				Result.RemoveAt(CommentStart, (CommentEnd + 2) - CommentStart, false);
+			}
+			else
+			{
+				// This looks like an error - an unclosed block comment.
+				break;
+			}
+		}
+		// Leftover line comments go next
+		for (CommentStart = Result.Find(TEXT("//~"), ESearchCase::CaseSensitive); CommentStart != INDEX_NONE; CommentStart = Result.Find(TEXT("//~"), ESearchCase::CaseSensitive))
+		{
+			CommentEnd = Result.Find(TEXT("\n"), ESearchCase::CaseSensitive, ESearchDir::FromStart, CommentStart);
+			if (CommentEnd != INDEX_NONE)
+			{
+				Result.RemoveAt(CommentStart, (CommentEnd + 1) - CommentStart, false);
+			}
+			else
+			{
+				Result.RemoveAt(CommentStart, Result.Len() - CommentStart, false);
+				break;
+			}
+		}
+		// Finish by shrinking if anything was removed, since we deferred this during the search.
+		Result.Shrink();
+	}
+
+	// Check for known commenting styles.
+	const bool bJavaDocStyle = Result.Contains(TEXT("/**"), ESearchCase::CaseSensitive);
+	const bool bCStyle = Result.Contains(TEXT("/*"), ESearchCase::CaseSensitive);
+	const bool bCPPStyle = Result.StartsWith(TEXT("//"), ESearchCase::CaseSensitive);
+
+	if ( bJavaDocStyle || bCStyle)
+	{
+		// Remove beginning and end markers.
+		if (bJavaDocStyle)
+		{
+			Result.ReplaceInline(TEXT("/**"), TEXT(""), ESearchCase::CaseSensitive);
+		}
+		if (bCStyle)
+		{
+			Result.ReplaceInline(TEXT("/*"), TEXT(""), ESearchCase::CaseSensitive);
+		}
+		Result.ReplaceInline(TEXT("*/"), TEXT(""), ESearchCase::CaseSensitive);
+	}
+
+	if ( bCPPStyle )
+	{
+		// Remove c++-style comment markers.  Also handle javadoc-style comments 
+		Result.ReplaceInline(TEXT("///"), TEXT(""), ESearchCase::CaseSensitive);
+		Result.ReplaceInline(TEXT("//"), TEXT(""), ESearchCase::CaseSensitive);
+
+		// Parser strips cpptext and replaces it with "// (cpptext)" -- prevent
+		// this from being treated as a comment on variables declared below the
+		// cpptext section
+		Result.ReplaceInline(TEXT("(cpptext)"), TEXT(""));
+	}
+
+	// Get rid of carriage return or tab characters, which mess up tooltips.
+	Result.ReplaceInline(TEXT( "\r" ), TEXT( "" ), ESearchCase::CaseSensitive);
+
+	//wx widgets has a hard coded tab size of 8
+	{
+		const int32 SpacesPerTab = 8;
+		Result.ConvertTabsToSpacesInline(SpacesPerTab);
+	}
+
+	// get rid of uniform leading whitespace and all trailing whitespace, on each line
+	TArray<FString> Lines;
+	Result.ParseIntoArray(Lines, TEXT("\n"), false);
+
+	for (FString& Line : Lines)
+	{
+		// Remove trailing whitespace
+		Line.TrimEndInline();
+
+		// Remove leading "*" and "* " in javadoc comments.
+		if (bJavaDocStyle)
+		{
+			// Find first non-whitespace character
+			int32 Pos = 0;
+			while (Pos < Line.Len() && FChar::IsWhitespace(Line[Pos]))
+			{
+				++Pos;
+			}
+
+			// Is it a *?
+			if (Pos < Line.Len() && Line[Pos] == '*')
+			{
+				// Eat next space as well
+				if (Pos+1 < Line.Len() && FChar::IsWhitespace(Line[Pos+1]))
+				{
+					++Pos;
+				}
+
+				Line.RightChopInline(Pos + 1, false);
+			}
+		}
+			}
+
+	auto IsWhitespaceOrLineSeparator = [](const FString& Line)
+	{
+		int32 LineLength = Line.Len();
+		int32 WhitespaceCount = 0;
+		while (WhitespaceCount < LineLength && FChar::IsWhitespace(Line[WhitespaceCount]))
+		{
+			++WhitespaceCount;
+		}
+
+		if (WhitespaceCount == LineLength)
+		{
+			return true;
+	}
+
+		const TCHAR* Str = (*Line) + WhitespaceCount;
+		return IsAllSameChar(Str, TEXT('-')) || IsAllSameChar(Str, TEXT('=')) || IsAllSameChar(Str, TEXT('*'));
+	};
+
+	// Find first meaningful line
+	int32 FirstIndex = 0;
+	for (const FString& Line : Lines)
+	{
+		if (!IsWhitespaceOrLineSeparator(Line))
+	{
+			break;
+		}
+
+		++FirstIndex;
+	}
+
+	int32 LastIndex = Lines.Num();
+	while (LastIndex != FirstIndex)
+	{
+		const FString& Line = Lines[LastIndex - 1];
+
+		if (!IsWhitespaceOrLineSeparator(Line))
+		{
+			break;
+		}
+
+		--LastIndex;
+	}
+
+	Result.Reset();
+
+	if (FirstIndex != LastIndex)
+	{
+		FString& FirstLine = Lines[FirstIndex];
+
+		// Figure out how much whitespace is on the first line
+		int32 MaxNumWhitespaceToRemove;
+		for (MaxNumWhitespaceToRemove = 0; MaxNumWhitespaceToRemove < FirstLine.Len(); MaxNumWhitespaceToRemove++)
+		{
+			if (!FChar::IsLinebreak(FirstLine[MaxNumWhitespaceToRemove]) && !FChar::IsWhitespace(FirstLine[MaxNumWhitespaceToRemove]))
+			{
+				break;
+			}
+		}
+
+		for (int32 Index = FirstIndex; Index != LastIndex; ++Index)
+		{
+			FString& Line = Lines[Index];
+
+			int32 TemporaryMaxWhitespace = MaxNumWhitespaceToRemove;
+
+			// Allow eating an extra tab on subsequent lines if it's present
+			if ((Index > 0) && (Line.Len() > 0) && (Line[0] == '\t'))
+			{
+				TemporaryMaxWhitespace++;
+			}
+
+			// Advance past whitespace
+			int32 Pos = 0;
+			while (Pos < TemporaryMaxWhitespace && Pos < Line.Len() && FChar::IsWhitespace(Line[Pos]))
+			{
+				++Pos;
+			}
+
+			if (Pos > 0)
+			{
+				Line.RightChopInline(Pos, false);
+			}
+
+			if (Index > 0)
+			{
+				Result += TEXT("\n");
+			}
+
+			if (Line.Len() && !IsAllSameChar(*Line, TEXT('=')))
+			{
+				Result += Line;
+			}
+		}
+	}
+
+	//@TODO: UCREMOVAL: Really want to trim an arbitrary number of newlines above and below, but keep multiple newlines internally
+	// Make sure it doesn't start with a newline
+	if (!Result.IsEmpty() && FChar::IsLinebreak(Result[0]))
+	{
+		Result.RightChopInline(1, false);
+	}
+
+	// Make sure it doesn't end with a dead newline
+	if (!Result.IsEmpty() && FChar::IsLinebreak(Result[Result.Len() - 1]))
+	{
+		Result.LeftInline(Result.Len() - 1, false);
+	}
+
+	// Done.
+	return Result;
+}
+
+TMap<FName, FString> FHeaderParser::GetParameterToolTipsFromFunctionComment(const FString& Input)
+{
+	SCOPE_SECONDS_COUNTER_UHT(DocumentationPolicy);
+
+	TMap<FName, FString> Map;
+	if (Input.IsEmpty())
+	{
+		return Map;
+	}
+	
+	TArray<FString> Params;
+	static const TCHAR ParamTag[] = TEXT("@param");
+	static const TCHAR ReturnTag[] = TEXT("@return");
+	static const TCHAR ReturnParamPrefix[] = TEXT("ReturnValue ");
+
+	/**
+	 * Search for @param / @return followed by a section until a line break.
+	 * For example: "@param Test MyTest Variable" becomes "Test", "MyTest Variable"
+	 * These pairs are then later split and stored as the parameter tooltips.
+	 * Once we don't find either @param or @return we break from the loop.
+	 */
+	int32 Offset = 0;
+	while (Offset < Input.Len())
+	{
+		const TCHAR* ParamPrefix = TEXT("");
+		int32 ParamStart = Input.Find(ParamTag, ESearchCase::CaseSensitive, ESearchDir::FromStart, Offset);
+		if(ParamStart != INDEX_NONE)
+		{
+			ParamStart = ParamStart + UE_ARRAY_COUNT(ParamTag);
+			Offset = ParamStart;
+		}
+		else
+		{
+			ParamStart = Input.Find(ReturnTag, ESearchCase::CaseSensitive, ESearchDir::FromStart, Offset);
+			if (ParamStart != INDEX_NONE)
+			{
+				ParamStart = ParamStart + UE_ARRAY_COUNT(ReturnTag);
+				Offset = ParamStart;
+				ParamPrefix = ReturnParamPrefix;
+			}
+			else
+			{
+				// no @param, no @return?
+				break;
+			}
+		}
+
+		int32 ParamEnd = Input.Find(TEXT("\n"), ESearchCase::CaseSensitive, ESearchDir::FromStart, ParamStart);
+		if (ParamEnd == INDEX_NONE)
+		{
+			ParamEnd = Input.Len();
+		}
+		Offset = ParamEnd;
+
+		Params.Add(ParamPrefix + Input.Mid(ParamStart, ParamEnd - ParamStart - 1));
+	}
+
+	for (FString& Param : Params)
+	{
+		Param.ConvertTabsToSpacesInline(4);
+		Param.TrimStartAndEndInline();
+
+		int32 FirstSpaceIndex = -1;
+		if (!Param.FindChar(' ', FirstSpaceIndex))
+		{
+			continue;
+		}
+
+		FString ParamToolTip = Param.Mid(FirstSpaceIndex + 1);
+		ParamToolTip.TrimStartInline();
+
+		Param.LeftInline(FirstSpaceIndex);
+
+		Map.Add(*Param, MoveTemp(ParamToolTip));
+	}
+
+	return Map;
+}
+
+
+void FHeaderParser::AddFormattedPrevCommentAsTooltipMetaData(TMap<FName, FString>& MetaData)
+{
+	// Don't add a tooltip if one already exists.
+	if (MetaData.Find(NAME_ToolTip))
+	{
+		return;
+	}
+
+	// Add the comment if it is not empty
+	if (!PrevComment.IsEmpty())
+	{
+		MetaData.Add(NAME_Comment, *PrevComment);
+	}
+
+	// Don't add a tooltip if the comment is empty after formatting.
+	FString FormattedComment = FormatCommentForToolTip(PrevComment);
+	if (!FormattedComment.Len())
+	{
+		return;
+	}
+
+	MetaData.Add(NAME_ToolTip, *FormattedComment);
+
+	// We've already used this comment as a tooltip, so clear it so that it doesn't get used again
+	PrevComment.Empty();
+}
+
+static const TCHAR* GetAccessSpecifierName(EAccessSpecifier AccessSpecifier)
+{
+	switch (AccessSpecifier)
+	{
+		case ACCESS_Public:
+			return TEXT("public");
+		case ACCESS_Protected:
+			return TEXT("protected");
+		case ACCESS_Private:
+			return TEXT("private");
+		default:
+			check(0);
+	}
+	return TEXT("");
+}
+
+// Tries to parse the token as an access protection specifier (public:, protected:, or private:)
+EAccessSpecifier FHeaderParser::ParseAccessProtectionSpecifier(const FToken& Token)
+{
+	EAccessSpecifier ResultAccessSpecifier = ACCESS_NotAnAccessSpecifier;
+
+	for (EAccessSpecifier Test = EAccessSpecifier(ACCESS_NotAnAccessSpecifier + 1); Test != ACCESS_Num; Test = EAccessSpecifier(Test + 1))
+	{
+		if (Token.Matches(GetAccessSpecifierName(Test), ESearchCase::CaseSensitive))
+		{
+			auto ErrorMessageGetter = [&Token]() { return FString::Printf(TEXT("after %s"), Token.Identifier);  };
+
+			// Consume the colon after the specifier
+			RequireSymbol(TEXT(':'), ErrorMessageGetter);
+			return Test;
+		}
+	}
+	return ACCESS_NotAnAccessSpecifier;
+}
+
+
+/**
+ * Compile a struct definition.
+ */
+UScriptStruct* FHeaderParser::CompileStructDeclaration(FClasses& AllClasses)
+{
+	FUnrealSourceFile* CurrentSrcFile = GetCurrentSourceFile();
+	TSharedPtr<FFileScope> Scope = CurrentSrcFile->GetScope();
+
+	// Make sure structs can be declared here.
+	CheckAllow( TEXT("'struct'"), ENestAllowFlags::TypeDecl );
+
+	FScriptLocation StructDeclaration;
+
+	bool IsNative = false;
+	bool IsExport = false;
+	bool IsTransient = false;
+	uint32 StructFlags = STRUCT_Native;
+	TMap<FName, FString> MetaData;
+
+	// Get the struct specifier list
+	TArray<FPropertySpecifier> SpecifiersFound;
+	ReadSpecifierSetInsideMacro(SpecifiersFound, TEXT("Struct"), MetaData);
+
+	// Consume the struct keyword
+	RequireIdentifier(TEXT("struct"), ESearchCase::CaseSensitive, TEXT("Struct declaration specifier"));
+
+	// The struct name as parsed in script and stripped of it's prefix
+	FString StructNameInScript;
+
+	// The struct name stripped of it's prefix
+	FString StructNameStripped;
+
+	// The required API module for this struct, if any
+	FString RequiredAPIMacroIfPresent;
+
+	// alignas() can come before or after the deprecation macro.
+	// We can't have both, but the compiler will catch that anyway.
+	SkipAlignasIfNecessary(*this);
+	SkipDeprecatedMacroIfNecessary(*this);
+	SkipAlignasIfNecessary(*this);
+
+	// Read the struct name
+	ParseNameWithPotentialAPIMacroPrefix(/*out*/ StructNameInScript, /*out*/ RequiredAPIMacroIfPresent, TEXT("struct"));
+
+	// Record that this struct is RequiredAPI if the CORE_API style macro was present
+	if (!RequiredAPIMacroIfPresent.IsEmpty())
+	{
+		StructFlags |= STRUCT_RequiredAPI;
+	}
+
+	StructNameStripped = GetClassNameWithPrefixRemoved(StructNameInScript);
+
+	// Effective struct name
+	const FString EffectiveStructName = *StructNameStripped;
+
+	// Verify that this struct name is not set to a reserved name	
+	if (FHeaderParser::IsReservedTypeName(EffectiveStructName))
+	{
+		FError::Throwf(TEXT("Struct '%s' uses a reserved type name ('%s')."), *StructNameInScript, *EffectiveStructName);
+	}
+
+	// Process the list of specifiers
+	for (const FPropertySpecifier& Specifier : SpecifiersFound)
+	{
+		switch ((EStructSpecifier)Algo::FindSortedStringCaseInsensitive(*Specifier.Key, GStructSpecifierStrings))
+		{
+			default:
+			{
+				FError::Throwf(TEXT("Unknown struct specifier '%s'"), *Specifier.Key);
+			}
+			break;
+
+			case EStructSpecifier::NoExport:
+			{
+				//UE_LOG_WARNING_UHT(TEXT("Struct named %s in %s is still marked noexport"), *EffectiveStructName, *(Class->GetName()));//@TODO: UCREMOVAL: Debug printing
+				StructFlags &= ~STRUCT_Native;
+				StructFlags |= STRUCT_NoExport;
+			}
+			break;
+
+			case EStructSpecifier::Atomic:
+			{
+				StructFlags |= STRUCT_Atomic;
+			}
+			break;
+
+			case EStructSpecifier::Immutable:
+			{
+				StructFlags |= STRUCT_Immutable | STRUCT_Atomic;
+
+				if (!FPaths::IsSamePath(Filename, GTypeDefinitionInfoMap[UObject::StaticClass()]->GetUnrealSourceFile().GetFilename()))
+				{
+					UE_LOG_ERROR_UHT(TEXT("Immutable is being phased out in favor of SerializeNative, and is only legal on the mirror structs declared in UObject"));
+				}
+			}
+			break;
+		}
+	}
+
+	// Verify uniqueness (if declared within UClass).
+	{
+		UField* Existing = Scope->FindTypeByName(*EffectiveStructName);
+		if (Existing)
+		{
+			FError::Throwf(TEXT("struct: '%s' already defined here"), *EffectiveStructName);
+		}
+
+		if (UStruct* FoundType = FindObject<UStruct>(ANY_PACKAGE, *EffectiveStructName))
+		{
+			if (TTuple<TSharedRef<FUnrealSourceFile>, int32>* FoundTypeInfo = GStructToSourceLine.Find(FoundType))
+			{
+				FError::Throwf(
+					TEXT("struct: '%s' conflicts with another type of the same name defined at %s(%d)"),
+					*EffectiveStructName,
+					*FoundTypeInfo->Get<0>()->GetFilename(),
+					FoundTypeInfo->Get<1>()
+				);
+			}
+			else
+			{
+				FError::Throwf(TEXT("struct: '%s' conflicts with another type of the same name"), *EffectiveStructName);
+			}
+
+		}
+	}
+
+	// Get optional superstruct.
+	bool bExtendsBaseStruct = false;
+	
+	if (MatchSymbol(TEXT(':')))
+	{
+		RequireIdentifier(TEXT("public"), ESearchCase::CaseSensitive, TEXT("struct inheritance"));
+		bExtendsBaseStruct = true;
+	}
+
+	UScriptStruct* BaseStruct = NULL;
+	if (bExtendsBaseStruct)
+	{
+		FToken ParentScope, ParentName;
+		if (GetIdentifier( ParentScope ))
+		{
+			RedirectTypeIdentifier(ParentScope);
+
+			TSharedPtr<FScope> StructScope = Scope;
+			FString ParentStructNameInScript = FString(ParentScope.Identifier);
+			if (MatchSymbol(TEXT('.')))
+			{
+				if (GetIdentifier(ParentName))
+				{
+					RedirectTypeIdentifier(ParentName);
+
+					ParentStructNameInScript = FString(ParentName.Identifier);
+					FString ParentNameStripped = GetClassNameWithPrefixRemoved(ParentScope.Identifier);
+					FClass* StructClass = AllClasses.FindClass(*ParentNameStripped);
+					if( !StructClass )
+					{
+						// If we find the literal class name, the user didn't use a prefix
+						StructClass = AllClasses.FindClass(ParentScope.Identifier);
+						if( StructClass )
+						{
+							FError::Throwf(TEXT("'struct': Parent struct class '%s' is missing a prefix, expecting '%s'"), ParentScope.Identifier, *FString::Printf(TEXT("%s%s"),StructClass->GetPrefixCPP(),ParentScope.Identifier) );
+						}
+						else
+						{
+							FError::Throwf(TEXT("'struct': Can't find parent struct class '%s'"), ParentScope.Identifier );
+						}
+					}
+
+					StructScope = FScope::GetTypeScope(StructClass);
+				}
+				else
+				{
+					FError::Throwf( TEXT("'struct': Missing parent struct type after '%s.'"), ParentScope.Identifier );
+				}
+			}
+			
+			FString ParentStructNameStripped;
+			const UField* Type = nullptr;
+			bool bOverrideParentStructName = false;
+
+			if( !StructsWithNoPrefix.Contains(ParentStructNameInScript) )
+			{
+				bOverrideParentStructName = true;
+				ParentStructNameStripped = GetClassNameWithPrefixRemoved(ParentStructNameInScript);
+			}
+
+			// If we're expecting a prefix, first try finding the correct field with the stripped struct name
+			if (bOverrideParentStructName)
+			{
+				Type = StructScope->FindTypeByName(*ParentStructNameStripped);
+			}
+
+			// If it wasn't found, try to find the literal name given
+			if (Type == nullptr)
+			{
+				Type = StructScope->FindTypeByName(*ParentStructNameInScript);
+			}
+
+			// Resolve structs declared in another class  //@TODO: UCREMOVAL: This seems extreme
+			if (Type == nullptr)
+			{
+				if (bOverrideParentStructName)
+				{
+					Type = FindObject<UScriptStruct>(ANY_PACKAGE, *ParentStructNameStripped);
+				}
+
+				if (Type == nullptr)
+				{
+					Type = FindObject<UScriptStruct>(ANY_PACKAGE, *ParentStructNameInScript);
+				}
+			}
+
+			// If the struct still wasn't found, throw an error
+			if (Type == NULL)
+			{
+				FError::Throwf(TEXT("'struct': Can't find struct '%s'"), *ParentStructNameInScript );
+			}
+			else
+			{
+				// If the struct was found, confirm it adheres to the correct syntax. This should always fail if we were expecting an override that was not found.
+				BaseStruct = ((UScriptStruct*)Type);
+				if( bOverrideParentStructName )
+				{
+					const TCHAR* PrefixCPP = StructsWithTPrefix.Contains(ParentStructNameStripped) ? TEXT("T") : BaseStruct->GetPrefixCPP();
+					if( ParentStructNameInScript != FString::Printf(TEXT("%s%s"), PrefixCPP, *ParentStructNameStripped) )
+					{
+						BaseStruct = nullptr;
+						FError::Throwf(TEXT("Parent Struct '%s' is missing a valid Unreal prefix, expecting '%s'"), *ParentStructNameInScript, *FString::Printf(TEXT("%s%s"), PrefixCPP, *Type->GetName()));
+					}
+				}
+			}
+		}
+		else
+		{
+			FError::Throwf(TEXT("'struct': Missing parent struct after ': public'") );
+		}
+	}
+
+	// if we have a base struct, propagate inherited struct flags now
+	if (BaseStruct != NULL)
+	{
+		StructFlags |= (BaseStruct->StructFlags&STRUCT_Inherit);
+	}
+	// Create.
+	UScriptStruct* Struct = new(EC_InternalUseOnlyConstructor, CurrentSrcFile->GetPackage(), *EffectiveStructName, RF_Public) UScriptStruct(FObjectInitializer(), BaseStruct);
+
+	Scope->AddType(Struct);
+	GTypeDefinitionInfoMap.Add(Struct, MakeShared<FUnrealTypeDefinitionInfo>(*CurrentSrcFile, InputLine));
+	FScope::AddTypeScope(Struct, &CurrentSrcFile->GetScope().Get());
+
+	AddModuleRelativePathToMetadata(Struct, MetaData);
+
+	// Check to make sure the syntactic native prefix was set-up correctly.
+	// If this check results in a false positive, it will be flagged as an identifier failure.
+	FString DeclaredPrefix = GetClassPrefix( StructNameInScript );
+	if( DeclaredPrefix == Struct->GetPrefixCPP() || DeclaredPrefix == TEXT("T") )
+	{
+		// Found a prefix, do a basic check to see if it's valid
+		const TCHAR* ExpectedPrefixCPP = StructsWithTPrefix.Contains(StructNameStripped) ? TEXT("T") : Struct->GetPrefixCPP();
+		FString ExpectedStructName = FString::Printf(TEXT("%s%s"), ExpectedPrefixCPP, *StructNameStripped);
+		if (StructNameInScript != ExpectedStructName)
+		{
+			FError::Throwf(TEXT("Struct '%s' has an invalid Unreal prefix, expecting '%s'"), *StructNameInScript, *ExpectedStructName);
+		}
+	}
+	else
+	{
+		const TCHAR* ExpectedPrefixCPP = StructsWithTPrefix.Contains(StructNameInScript) ? TEXT("T") : Struct->GetPrefixCPP();
+		FString ExpectedStructName = FString::Printf(TEXT("%s%s"), ExpectedPrefixCPP, *StructNameInScript);
+		FError::Throwf(TEXT("Struct '%s' is missing a valid Unreal prefix, expecting '%s'"), *StructNameInScript, *ExpectedStructName);
+	}
+
+	Struct->StructFlags = EStructFlags(Struct->StructFlags | StructFlags);
+
+	AddFormattedPrevCommentAsTooltipMetaData(MetaData);
+
+	// Register the metadata
+	AddMetaDataToClassData(Struct, MoveTemp(MetaData));
+
+	// Get opening brace.
+	RequireSymbol( TEXT('{'), TEXT("'struct'") );
+
+	// Members of structs have a default public access level in c++
+	// Assume that, but restore the parser state once we finish parsing this struct
+	TGuardValue<EAccessSpecifier> HoldFromClass(CurrentAccessSpecifier, ACCESS_Public);
+
+	{
+		FToken StructToken;
+		StructToken.Struct = Struct;
+
+		// add this struct to the compiler's persistent tracking system
+		FClassMetaData* ClassMetaData = GScriptHelper.AddClassData(StructToken.Struct, CurrentSrcFile);
+	}
+
+	int32 SavedLineNumber = InputLine;
+
+	// Clear comment before parsing body of the struct.
+	
+
+	// Parse all struct variables.
+	FToken Token;
+	while (1)
+	{
+		ClearComment();
+		GetToken( Token );
+
+		if (EAccessSpecifier AccessSpecifier = ParseAccessProtectionSpecifier(Token))
+		{
+			CurrentAccessSpecifier = AccessSpecifier;
+		}
+		else if (Token.Matches(TEXT("UPROPERTY"), ESearchCase::CaseSensitive))
+		{
+			CompileVariableDeclaration(AllClasses, Struct);
+		}
+		else if (Token.Matches(TEXT("UFUNCTION"), ESearchCase::CaseSensitive))
+		{
+			FError::Throwf(TEXT("USTRUCTs cannot contain UFUNCTIONs."));
+		}
+		else if (Token.Matches(TEXT("RIGVM_METHOD"), ESearchCase::CaseSensitive))
+		{
+			CompileRigVMMethodDeclaration(AllClasses, Struct);
+		}
+		else if (Token.Matches(TEXT("GENERATED_USTRUCT_BODY"), ESearchCase::CaseSensitive) || Token.Matches(TEXT("GENERATED_BODY"), ESearchCase::CaseSensitive))
+		{
+			// Match 'GENERATED_USTRUCT_BODY' '(' [StructName] ')' or 'GENERATED_BODY' '(' [StructName] ')'
+			if (CurrentAccessSpecifier != ACCESS_Public)
+			{
+				FError::Throwf(TEXT("%s must be in the public scope of '%s', not private or protected."), Token.Identifier, *StructNameInScript);
+			}
+
+			if (Struct->StructMacroDeclaredLineNumber != INDEX_NONE)
+			{
+				FError::Throwf(TEXT("Multiple %s declarations found in '%s'"), Token.Identifier, *StructNameInScript);
+			}
+
+			Struct->StructMacroDeclaredLineNumber = InputLine;
+			RequireSymbol(TEXT('('), TEXT("'struct'"));
+
+			CompileVersionDeclaration(Struct);
+
+			RequireSymbol(TEXT(')'), TEXT("'struct'"));
+
+			// Eat a semicolon if present (not required)
+			SafeMatchSymbol(TEXT(';'));
+		}
+		else if ( Token.Matches(TEXT('#')) && MatchIdentifier(TEXT("ifdef"), ESearchCase::CaseSensitive) )
+		{
+			PushCompilerDirective(ECompilerDirective::Insignificant);
+		}
+		else if ( Token.Matches(TEXT('#')) && MatchIdentifier(TEXT("ifndef"), ESearchCase::CaseSensitive) )
+		{
+			PushCompilerDirective(ECompilerDirective::Insignificant);
+		}
+		else if (Token.Matches(TEXT('#')) && MatchIdentifier(TEXT("endif"), ESearchCase::CaseSensitive))
+		{
+			if (CompilerDirectiveStack.Num() < 1)
+			{
+				FError::Throwf(TEXT("Unmatched '#endif' in class or global scope"));
+			}
+			CompilerDirectiveStack.Pop();
+			// Do nothing and hope that the if code below worked out OK earlier
+		}
+		else if ( Token.Matches(TEXT('#')) && MatchIdentifier(TEXT("if"), ESearchCase::CaseSensitive) )
+		{
+			//@TODO: This parsing should be combined with CompileDirective and probably happen much much higher up!
+			bool bInvertConditional = MatchSymbol(TEXT('!'));
+			bool bConsumeAsCppText = false;
+
+			if (MatchIdentifier(TEXT("WITH_EDITORONLY_DATA"), ESearchCase::CaseSensitive) )
+			{
+				if (bInvertConditional)
+				{
+					FError::Throwf(TEXT("Cannot use !WITH_EDITORONLY_DATA"));
+				}
+
+				PushCompilerDirective(ECompilerDirective::WithEditorOnlyData);
+			}
+			else if (MatchIdentifier(TEXT("WITH_EDITOR"), ESearchCase::CaseSensitive) )
+			{
+				if (bInvertConditional)
+				{
+					FError::Throwf(TEXT("Cannot use !WITH_EDITOR"));
+				}
+				PushCompilerDirective(ECompilerDirective::WithEditor);
+			}
+			else if (MatchIdentifier(TEXT("CPP"), ESearchCase::CaseSensitive) || MatchConstInt(TEXT("0")) || MatchConstInt(TEXT("1")) || MatchIdentifier(TEXT("WITH_HOT_RELOAD"), ESearchCase::CaseSensitive) || MatchIdentifier(TEXT("WITH_HOT_RELOAD_CTORS"), ESearchCase::CaseSensitive))
+			{
+				bConsumeAsCppText = !bInvertConditional;
+				PushCompilerDirective(ECompilerDirective::Insignificant);
+			}
+			else
+			{
+				FError::Throwf(TEXT("'struct': Unsupported preprocessor directive inside a struct.") );
+			}
+
+			if (bConsumeAsCppText)
+			{
+				// Skip over the text, it is not recorded or processed
+				int32 nest = 1;
+				while (nest > 0)
+				{
+					TCHAR ch = GetChar(1);
+
+					if ( ch==0 )
+					{
+						FError::Throwf(TEXT("Unexpected end of struct definition %s"), *Struct->GetName());
+					}
+					else if ( ch=='{' || (ch=='#' && (PeekIdentifier(TEXT("if"), ESearchCase::CaseSensitive) || PeekIdentifier(TEXT("ifdef"), ESearchCase::CaseSensitive))) )
+					{
+						nest++;
+					}
+					else if ( ch=='}' || (ch=='#' && PeekIdentifier(TEXT("endif"), ESearchCase::CaseSensitive)) )
+					{
+						nest--;
+					}
+
+					if (nest==0)
+					{
+						RequireIdentifier(TEXT("endif"), ESearchCase::CaseSensitive, TEXT("'if'"));
+					}
+				}
+			}
+		}
+		else if (Token.Matches(TEXT('#')) && MatchIdentifier(TEXT("pragma"), ESearchCase::CaseSensitive))
+		{
+			// skip it and skip over the text, it is not recorded or processed
+			TCHAR c;
+			while (!IsEOL(c = GetChar()))
+			{
+			}
+		}
+		else if (ProbablyAnUnknownObjectLikeMacro(*this, Token))
+		{
+			// skip it
+		}
+		else
+		{
+			if (!Token.Matches( TEXT('}')))
+			{
+				// Skip declaration will destroy data in Token, so cache off the identifier in case we need to provfide an error
+				TCHAR FirstTokenIdentifier[NAME_SIZE];
+				FCString::Strncpy(FirstTokenIdentifier, Token.Identifier, NAME_SIZE);
+
+				if (!SkipDeclaration(Token))
+				{
+					FError::Throwf(TEXT("'struct': Unexpected '%s'"), FirstTokenIdentifier);
+				}	
+			}
+			else
+			{
+				MatchSemi();
+				break;
+			}
+		}
+	}
+
+	// Validation
+	bool bStructBodyFound = Struct->StructMacroDeclaredLineNumber != INDEX_NONE;
+	bool bExported        = !!(StructFlags & STRUCT_Native);
+	if (!bStructBodyFound && bExported)
+	{
+		// Roll the line number back to the start of the struct body and error out
+		InputLine = SavedLineNumber;
+		FError::Throwf(TEXT("Expected a GENERATED_BODY() at the start of struct"));
+	}
+
+	// Validate sparse class data
+	CheckSparseClassData(Struct);
+
+	// Link the properties within the struct
+	Struct->StaticLink(true);
+
+	return Struct;
+}
+
+/*-----------------------------------------------------------------------------
+	Retry management.
+-----------------------------------------------------------------------------*/
+
+/**
+ * Remember the current compilation points, both in the source being
+ * compiled and the object code being emitted.
+ *
+ * @param	Retry	[out] filled in with current compiler position information
+ */
+void FHeaderParser::InitScriptLocation( FScriptLocation& Retry )
+{
+	Retry.Input = Input;
+	Retry.InputPos = InputPos;
+	Retry.InputLine	= InputLine;
+}
+
+/**
+ * Return to a previously-saved retry point.
+ *
+ * @param	Retry	the point to return to
+ * @param	Binary	whether to modify the compiled bytecode
+ * @param	bText	whether to modify the compiler's current location in the text
+ */
+void FHeaderParser::ReturnToLocation(const FScriptLocation& Retry, bool Binary, bool bText)
+{
+	if (bText)
+	{
+		Input = Retry.Input;
+		InputPos = Retry.InputPos;
+		InputLine = Retry.InputLine;
+	}
+}
+
+/*-----------------------------------------------------------------------------
+	Nest information.
+-----------------------------------------------------------------------------*/
+
+//
+// Return the name for a nest type.
+//
+const TCHAR *FHeaderParser::NestTypeName( ENestType NestType )
+{
+	switch( NestType )
+	{
+		case ENestType::GlobalScope:
+			return TEXT("Global Scope");
+		case ENestType::Class:
+			return TEXT("Class");
+		case ENestType::NativeInterface:
+		case ENestType::Interface:
+			return TEXT("Interface");
+		case ENestType::FunctionDeclaration:
+			return TEXT("Function");
+		default:
+			check(false);
+			return TEXT("Unknown");
+	}
+}
+
+// Checks to see if a particular kind of command is allowed on this nesting level.
+bool FHeaderParser::IsAllowedInThisNesting(ENestAllowFlags AllowFlags)
+{
+	return (TopNest->Allow & AllowFlags) != ENestAllowFlags::None;
+}
+
+//
+// Make sure that a particular kind of command is allowed on this nesting level.
+// If it's not, issues a compiler error referring to the token and the current
+// nesting level.
+//
+void FHeaderParser::CheckAllow( const TCHAR* Thing, ENestAllowFlags AllowFlags )
+{
+	if (!IsAllowedInThisNesting(AllowFlags))
+	{
+		if (TopNest->NestType == ENestType::GlobalScope)
+		{
+			FError::Throwf(TEXT("%s is not allowed before the Class definition"), Thing );
+		}
+		else
+		{
+			FError::Throwf(TEXT("%s is not allowed here"), Thing );
+		}
+	}
+}
+
+bool FHeaderParser::AllowReferenceToClass(UStruct* Scope, UClass* CheckClass) const
+{
+	check(CheckClass);
+
+	return	(Scope->GetOutermost() == CheckClass->GetOutermost())
+		|| ((CheckClass->ClassFlags&CLASS_Parsed) != 0)
+		|| ((CheckClass->ClassFlags&CLASS_Intrinsic) != 0);
+}
+
+/*-----------------------------------------------------------------------------
+	Nest management.
+-----------------------------------------------------------------------------*/
+
+void FHeaderParser::PushNest(ENestType NestType, UStruct* InNode, FUnrealSourceFile* SourceFile)
+{
+	// Update pointer to top nesting level.
+	TopNest = &Nest[NestLevel++];
+	TopNest->SetScope(NestType == ENestType::GlobalScope ? &SourceFile->GetScope().Get() : &FScope::GetTypeScope(InNode).Get());
+	TopNest->NestType = NestType;
+
+	// Prevent overnesting.
+	if (NestLevel >= MAX_NEST_LEVELS)
+	{
+		FError::Throwf(TEXT("Maximum nesting limit exceeded"));
+	}
+
+	// Inherit info from stack node above us.
+	if (NestLevel > 1 && NestType == ENestType::GlobalScope)
+	{
+		// Use the existing stack node.
+		TopNest->SetScope(TopNest[-1].GetScope());
+	}
+
+	// NestType specific logic.
+	switch (NestType)
+	{
+	case ENestType::GlobalScope:
+		TopNest->Allow = ENestAllowFlags::Class | ENestAllowFlags::TypeDecl | ENestAllowFlags::ImplicitDelegateDecl;
+		break;
+
+	case ENestType::Class:
+		TopNest->Allow = ENestAllowFlags::VarDecl | ENestAllowFlags::Function | ENestAllowFlags::ImplicitDelegateDecl;
+		break;
+
+	case ENestType::NativeInterface:
+	case ENestType::Interface:
+		TopNest->Allow = ENestAllowFlags::Function;
+		break;
+
+	case ENestType::FunctionDeclaration:
+		TopNest->Allow = ENestAllowFlags::VarDecl;
+
+		break;
+
+	default:
+		FError::Throwf(TEXT("Internal error in PushNest, type %i"), (uint8)NestType);
+		break;
+	}
+}
+
+/**
+ * Decrease the nesting level and handle any errors that result.
+ *
+ * @param	NestType	nesting type of the current node
+ * @param	Descr		text to use in error message if any errors are encountered
+ */
+void FHeaderParser::PopNest(ENestType NestType, const TCHAR* Descr)
+{
+	// Validate the nesting state.
+	if (NestLevel <= 0)
+	{
+		FError::Throwf(TEXT("Unexpected '%s' at global scope"), Descr, NestTypeName(NestType));
+	}
+	else if (TopNest->NestType != NestType)
+	{
+		FError::Throwf(TEXT("Unexpected end of %s in '%s' block"), Descr, NestTypeName(TopNest->NestType));
+	}
+
+	if (NestType != ENestType::GlobalScope && NestType != ENestType::Class && NestType != ENestType::Interface && NestType != ENestType::NativeInterface && NestType != ENestType::FunctionDeclaration)
+	{
+		FError::Throwf(TEXT("Bad first pass NestType %i"), (uint8)NestType);
+	}
+
+	bool bLinkProps = true;
+	if (NestType == ENestType::Class)
+	{
+		UClass* TopClass = GetCurrentClass();
+		bLinkProps = !TopClass->HasAnyClassFlags(CLASS_Intrinsic);
+	}
+
+	if (NestType != ENestType::GlobalScope)
+	{
+		GetCurrentClass()->StaticLink(bLinkProps);
+	}
+
+	// Pop the nesting level.
+	NestType = TopNest->NestType;
+	NestLevel--;
+	if (NestLevel == 0)
+	{
+		TopNest = nullptr;
+	}
+	else
+	{
+		TopNest--;
+		check(TopNest >= Nest);
+
+	}
+}
+
+void FHeaderParser::FixupDelegateProperties( FClasses& AllClasses, UStruct* Struct, FScope& Scope, TMap<FName, UFunction*>& DelegateCache )
+{
+	check(Struct);
+
+	for ( FField* Field = Struct->ChildProperties; Field; Field = Field->Next )
+	{
+		FProperty* Property = CastField<FProperty>(Field);
+		if ( Property != NULL )
+		{
+			FDelegateProperty* DelegateProperty = CastField<FDelegateProperty>(Property);
+			FMulticastDelegateProperty* MulticastDelegateProperty = CastField<FMulticastDelegateProperty>(Property);
+			if ( DelegateProperty == NULL && MulticastDelegateProperty == NULL )
+			{
+				// if this is an array property, see if the array's type is a delegate
+				FArrayProperty* ArrayProp = CastField<FArrayProperty>(Property);
+				if ( ArrayProp != NULL )
+				{
+					DelegateProperty = CastField<FDelegateProperty>(ArrayProp->Inner);
+					MulticastDelegateProperty = CastField<FMulticastDelegateProperty>(ArrayProp->Inner);
+				}
+			}
+			if (DelegateProperty != nullptr || MulticastDelegateProperty != nullptr)
+			{
+				// this FDelegateProperty corresponds to an actual delegate variable (i.e. delegate<SomeDelegate> Foo); we need to lookup the token data for
+				// this property and verify that the delegate property's "type" is an actual delegate function
+				FClassMetaData* StructData = GScriptHelper.FindClassData(Struct);
+				check(StructData);
+				FTokenData* DelegatePropertyToken = StructData->FindTokenData(Property);
+				check(DelegatePropertyToken);
+
+				// attempt to find the delegate function in the map of functions we've already found
+				UFunction* SourceDelegateFunction = DelegateCache.FindRef(DelegatePropertyToken->Token.DelegateName);
+				if (SourceDelegateFunction == nullptr)
+				{
+					FString NameOfDelegateFunction = DelegatePropertyToken->Token.DelegateName.ToString() + FString( HEADER_GENERATED_DELEGATE_SIGNATURE_SUFFIX );
+					if ( !NameOfDelegateFunction.Contains(TEXT(".")) )
+					{
+						// an unqualified delegate function name - search for a delegate function by this name within the current scope
+						SourceDelegateFunction = Cast<UFunction>(Scope.FindTypeByName(*NameOfDelegateFunction));
+						if (SourceDelegateFunction == nullptr)
+						{
+							// Try to find in other packages.
+							UObject* DelegateSignatureOuter = DelegatePropertyToken->Token.DelegateSignatureOwnerClass 
+								? ((UObject*)DelegatePropertyToken->Token.DelegateSignatureOwnerClass) 
+								: ((UObject*)ANY_PACKAGE);
+							SourceDelegateFunction = Cast<UFunction>(StaticFindObject(UFunction::StaticClass(), DelegateSignatureOuter, *NameOfDelegateFunction));
+
+							if (SourceDelegateFunction == nullptr)
+							{
+								// convert this into a fully qualified path name for the error message.
+								NameOfDelegateFunction = Scope.GetName().ToString() + TEXT(".") + NameOfDelegateFunction;
+							}
+						}
+					}
+					else
+					{
+						FString DelegateClassName, DelegateName;
+						NameOfDelegateFunction.Split(TEXT("."), &DelegateClassName, &DelegateName);
+
+						// verify that we got a valid string for the class name
+						if ( DelegateClassName.Len() == 0 )
+						{
+							UngetToken(DelegatePropertyToken->Token);
+							FError::Throwf(TEXT("Invalid scope specified in delegate property function reference: '%s'"), *NameOfDelegateFunction);
+						}
+
+						// verify that we got a valid string for the name of the function
+						if ( DelegateName.Len() == 0 )
+						{
+							UngetToken(DelegatePropertyToken->Token);
+							FError::Throwf(TEXT("Invalid delegate name specified in delegate property function reference '%s'"), *NameOfDelegateFunction);
+						}
+
+						// make sure that the class that contains the delegate can be referenced here
+						UClass* DelegateOwnerClass = AllClasses.FindScriptClassOrThrow(DelegateClassName);
+						if (FScope::GetTypeScope(DelegateOwnerClass)->FindTypeByName(*DelegateName) != nullptr)
+						{
+							FError::Throwf(TEXT("Inaccessible type: '%s'"), *DelegateOwnerClass->GetPathName());
+						}
+						SourceDelegateFunction = Cast<UFunction>(FindField(DelegateOwnerClass, *DelegateName, false, UFunction::StaticClass(), NULL));	
+					}
+
+					if ( SourceDelegateFunction == NULL )
+					{
+						UngetToken(DelegatePropertyToken->Token);
+						FError::Throwf(TEXT("Failed to find delegate function '%s'"), *NameOfDelegateFunction);
+					}
+					else if ( (SourceDelegateFunction->FunctionFlags&FUNC_Delegate) == 0 )
+					{
+						UngetToken(DelegatePropertyToken->Token);
+						FError::Throwf(TEXT("Only delegate functions can be used as the type for a delegate property; '%s' is not a delegate."), *NameOfDelegateFunction);
+					}
+				}
+
+				// successfully found the delegate function that this delegate property corresponds to
+
+				// save this into the delegate cache for faster lookup later
+				DelegateCache.Add(DelegatePropertyToken->Token.DelegateName, SourceDelegateFunction);
+
+				// bind it to the delegate property
+				if( DelegateProperty != NULL )
+				{
+					if( !SourceDelegateFunction->HasAnyFunctionFlags( FUNC_MulticastDelegate ) )
+					{
+						DelegateProperty->SignatureFunction = DelegatePropertyToken->Token.Function = SourceDelegateFunction;
+					}
+					else
+					{
+						FError::Throwf(TEXT("Unable to declare a single-cast delegate property for a multi-cast delegate type '%s'.  Either add a 'multicast' qualifier to the property or change the delegate type to be single-cast as well."), *SourceDelegateFunction->GetName());
+					}
+				}
+				else if( MulticastDelegateProperty != NULL )
+				{
+					if( SourceDelegateFunction->HasAnyFunctionFlags( FUNC_MulticastDelegate ) )
+					{
+						MulticastDelegateProperty->SignatureFunction = DelegatePropertyToken->Token.Function = SourceDelegateFunction;
+
+						if(MulticastDelegateProperty->HasAnyPropertyFlags(CPF_BlueprintAssignable | CPF_BlueprintCallable))
+						{
+							for (TFieldIterator<FProperty> PropIt(SourceDelegateFunction); PropIt && (PropIt->PropertyFlags & CPF_Parm); ++PropIt)
+							{
+								FProperty* FuncParam = *PropIt;
+
+								if (!IsPropertySupportedByBlueprint(FuncParam, false))
+								{
+									FString ExtendedCPPType;
+									FString CPPType = FuncParam->GetCPPType(&ExtendedCPPType);
+									UE_LOG_ERROR_UHT(TEXT("Type '%s%s' is not supported by blueprint. %s.%s"), *CPPType, *ExtendedCPPType, *SourceDelegateFunction->GetName(), *FuncParam->GetName());
+								}
+
+								if(FuncParam->HasAllPropertyFlags(CPF_OutParm) && !FuncParam->HasAllPropertyFlags(CPF_ConstParm)  )
+								{
+									const bool bClassGeneratedFromBP = FClass::IsDynamic(Struct);
+									const bool bAllowedArrayRefFromBP = bClassGeneratedFromBP && FuncParam->IsA<FArrayProperty>();
+									if (!bAllowedArrayRefFromBP)
+									{
+										UE_LOG_ERROR_UHT(TEXT("BlueprintAssignable delegates do not support non-const references at the moment. Function: %s Parameter: '%s'"), *SourceDelegateFunction->GetName(), *FuncParam->GetName());
+									}
+								}
+							}
+						}
+
+					}
+					else
+					{
+						FError::Throwf(TEXT("Unable to declare a multi-cast delegate property for a single-cast delegate type '%s'.  Either remove the 'multicast' qualifier from the property or change the delegate type to be 'multicast' as well."), *SourceDelegateFunction->GetName());
+					}
+				}
+			}
+		}
+	}
+
+	for (UField* Field = Struct->Children; Field; Field = Field->Next)
+	{
+		// if this is a state, function, or script struct, it might have its own delegate properties which need to be validated
+		UStruct* InternalStruct = Cast<UStruct>(Field);
+		if ( InternalStruct != NULL )
+		{
+			FixupDelegateProperties(AllClasses, InternalStruct, Scope, DelegateCache);
+		}
+	}
+
+	CheckDocumentationPolicyForStruct(Struct);
+
+	ParseRigVMMethodParameters(Struct);
+}
+
+void FHeaderParser::CheckSparseClassData(const UStruct* StructToCheck)
+{
+	// we're looking for classes that have sparse class data structures
+	const UClass* ClassToCheck = Cast<const UClass>(StructToCheck);
+	if (!ClassToCheck)
+	{
+		// make sure we don't try to have sparse class data inside of a struct instead of a class
+		if (StructToCheck->HasMetaData(FHeaderParserNames::NAME_SparseClassDataTypes))
+		{
+			FError::Throwf(TEXT("%s contains sparse class data but is not a class."), *StructToCheck->GetName());
+		}
+		return;
+	}
+
+	if (!ClassToCheck->HasMetaData(FHeaderParserNames::NAME_SparseClassDataTypes))
+	{
+		return;
+	}
+
+	TArray<FString> SparseClassDataTypes;
+	((FClass*)ClassToCheck)->GetSparseClassDataTypes(SparseClassDataTypes);
+
+	// for now we only support one sparse class data structure per class
+	if (SparseClassDataTypes.Num() > 1)
+	{
+		FError::Throwf(TEXT("Class %s contains multiple sparse class data types."), *ClassToCheck->GetName());
+		return;
+	}
+	if (SparseClassDataTypes.Num() == 0)
+	{
+		FError::Throwf(TEXT("Class %s has sparse class metadata but does not specify a type."), *ClassToCheck->GetName());
+		return;
+	}
+
+	for (const FString& SparseClassDataTypeName : SparseClassDataTypes)
+	{
+		UScriptStruct* SparseClassDataStruct = FindObjectSafe<UScriptStruct>(ANY_PACKAGE, *SparseClassDataTypeName);
+
+		// make sure the sparse class data struct actually exists
+		if (!SparseClassDataStruct)
+		{
+			FError::Throwf(TEXT("Unable to find sparse data type %s for class %s."), *SparseClassDataTypeName, *ClassToCheck->GetName());
+			return;
+		}
+
+		// check the data struct for invalid properties
+		for (TFieldIterator<FProperty> Property(SparseClassDataStruct); Property; ++Property)
+		{
+			if (Property->HasAnyPropertyFlags(CPF_BlueprintAssignable))
+			{
+				FError::Throwf(TEXT("Sparse class data types can not contain blueprint assignable delegates. Type '%s' Delegate '%s'"), *SparseClassDataStruct->GetName(), *Property->GetName());
+			}
+
+			// all sparse properties should have EditDefaultsOnly
+			if (!Property->HasAllPropertyFlags(CPF_Edit | CPF_DisableEditOnInstance))
+			{
+				FError::Throwf(TEXT("Sparse class data types must be EditDefaultsOnly. Type '%s' Property '%s'"), *SparseClassDataStruct->GetName(), *Property->GetName());
+			}
+
+			// no sparse properties should have BlueprintReadWrite
+			if (Property->HasAllPropertyFlags(CPF_BlueprintVisible) && !Property->HasAllPropertyFlags(CPF_BlueprintReadOnly))
+			{
+				FError::Throwf(TEXT("Sparse class data types must not be BlueprintReadWrite. Type '%s' Property '%s'"), *SparseClassDataStruct->GetName(), *Property->GetName());
+			}
+		}
+
+		// if the class's parent has a sparse class data struct then the current class must also use the same struct or one that inherits from it
+		const UClass* ParentClass = ClassToCheck->GetSuperClass();
+		TArray<FString> ParentSparseClassDataTypeNames;
+		((FClass*)ParentClass)->GetSparseClassDataTypes(ParentSparseClassDataTypeNames);
+		for (FString& ParentSparseClassDataTypeName : ParentSparseClassDataTypeNames)
+		{
+			UScriptStruct* ParentSparseClassDataStruct = FindObjectSafe<UScriptStruct>(ANY_PACKAGE, *ParentSparseClassDataTypeName);
+			if (ParentSparseClassDataStruct && !SparseClassDataStruct->IsChildOf(ParentSparseClassDataStruct))
+			{
+				FError::Throwf(TEXT("Class %s is a child of %s but its sparse class data struct, %s, does not inherit from %s."), *ClassToCheck->GetName(), *ParentClass->GetName(), *SparseClassDataStruct->GetName(), *ParentSparseClassDataStruct->GetName());
+			}
+		}
+	}
+}
+
+void FHeaderParser::VerifyBlueprintPropertyGetter(FProperty* Prop, UFunction* TargetFunc)
+{
+	check(TargetFunc);
+
+	FProperty* ReturnProp = TargetFunc->GetReturnProperty();
+	if (TargetFunc->NumParms > 1 || (TargetFunc->NumParms == 1 && ReturnProp == nullptr))
+	{
+		UE_LOG_ERROR_UHT(TEXT("Blueprint Property getter function %s must not have parameters."), *TargetFunc->GetName());
+	}
+
+	if (ReturnProp == nullptr || !Prop->SameType(ReturnProp))
+	{
+		FString ExtendedCPPType;
+		FString CPPType = Prop->GetCPPType(&ExtendedCPPType);
+		UE_LOG_ERROR_UHT(TEXT("Blueprint Property getter function %s must have return value of type %s%s."), *TargetFunc->GetName(), *CPPType, *ExtendedCPPType);
+	}
+
+	if (TargetFunc->HasAnyFunctionFlags(FUNC_Event))
+	{
+		UE_LOG_ERROR_UHT(TEXT("Blueprint Property setter function cannot be a blueprint event."));
+	}
+	else if (!TargetFunc->HasAnyFunctionFlags(FUNC_BlueprintPure))
+	{
+		UE_LOG_ERROR_UHT(TEXT("Blueprint Property getter function must be pure."));
+	}
+}
+
+void FHeaderParser::VerifyBlueprintPropertySetter(FProperty* Prop, UFunction* TargetFunc)
+{
+	check(TargetFunc);
+	FProperty* ReturnProp = TargetFunc->GetReturnProperty();
+
+	if (ReturnProp)
+	{
+		UE_LOG_ERROR_UHT(TEXT("Blueprint Property setter function %s must not have a return value."), *TargetFunc->GetName());
+	}
+	else
+	{
+		TFieldIterator<FProperty> Parm(TargetFunc);
+		if (TargetFunc->NumParms != 1 || !Prop->SameType(*Parm))
+		{
+			FString ExtendedCPPType;
+			FString CPPType = Prop->GetCPPType(&ExtendedCPPType);
+			UE_LOG_ERROR_UHT(TEXT("Blueprint Property setter function %s must have exactly one parameter of type %s%s."), *TargetFunc->GetName(), *CPPType, *ExtendedCPPType);
+		}
+	}
+
+	if (TargetFunc->HasAnyFunctionFlags(FUNC_Event))
+	{
+		UE_LOG_ERROR_UHT(TEXT("Blueprint Property setter function cannot be a blueprint event."));
+	}
+	else if (!TargetFunc->HasAnyFunctionFlags(FUNC_BlueprintCallable))
+	{
+		UE_LOG_ERROR_UHT(TEXT("Blueprint Property setter function must be blueprint callable."));
+	}
+	else if (TargetFunc->HasAnyFunctionFlags(FUNC_BlueprintPure))
+	{
+		UE_LOG_ERROR_UHT(TEXT("Blueprint Property setter function must not be pure."));
+	}
+}
+
+void FHeaderParser::VerifyRepNotifyCallback(FProperty* Prop, UFunction* TargetFunc)
+{
+	if( TargetFunc )
+	{
+		if (TargetFunc->GetReturnProperty())
+		{
+			UE_LOG_ERROR_UHT(TEXT("Replication notification function %s must not have return value."), *TargetFunc->GetName());
+		}
+
+		const bool bIsArrayProperty = ( Prop->ArrayDim > 1 || CastField<FArrayProperty>(Prop) );
+		const int32 MaxParms = bIsArrayProperty ? 2 : 1;
+
+		if ( TargetFunc->NumParms > MaxParms)
+		{
+			UE_LOG_ERROR_UHT(TEXT("Replication notification function %s has too many parameters."), *TargetFunc->GetName());
+		}
+
+		TFieldIterator<FProperty> Parm(TargetFunc);
+		if ( TargetFunc->NumParms >= 1 && Parm)
+		{
+			// First parameter is always the old value:
+			if ( !Prop->SameType(*Parm) )
+			{
+				FString ExtendedCPPType;
+				FString CPPType = Prop->GetCPPType(&ExtendedCPPType);
+				UE_LOG_ERROR_UHT(TEXT("Replication notification function %s has invalid parameter for property %s. First (optional) parameter must be of type %s%s."), *TargetFunc->GetName(), *Prop->GetName(), *CPPType, *ExtendedCPPType);
+			}
+
+			++Parm;
+		}
+
+		if ( TargetFunc->NumParms >= 2 && Parm)
+		{
+			// A 2nd parameter for arrays can be specified as a const TArray<uint8>&. This is a list of element indices that have changed
+			FArrayProperty *ArrayProp = CastField<FArrayProperty>(*Parm);
+			if (!(ArrayProp && CastField<FByteProperty>(ArrayProp->Inner)) || !(Parm->GetPropertyFlags() & CPF_ConstParm) || !(Parm->GetPropertyFlags() & CPF_ReferenceParm))
+			{
+				UE_LOG_ERROR_UHT(TEXT("Replication notification function %s (optional) second parameter must be of type 'const TArray<uint8>&'"), *TargetFunc->GetName());
+			}
+		}
+	}
+	else
+	{
+		// Couldn't find a valid function...
+		UE_LOG_ERROR_UHT(TEXT("Replication notification function %s not found"), *Prop->RepNotifyFunc.ToString() );
+	}
+}
+void FHeaderParser::VerifyPropertyMarkups( UClass* TargetClass )
+{
+	// Iterate over all properties, looking for those flagged as CPF_RepNotify
+	for ( FField* Field = TargetClass->ChildProperties; Field; Field = Field->Next )
+	{
+		if (FProperty* Prop = CastField<FProperty>(Field))
+		{
+			auto FindTargetFunction = [&](const FName FuncName)
+			{
+				// Search through this class and its superclasses looking for the specified callback
+				UFunction* TargetFunc = nullptr;
+				UClass* SearchClass = TargetClass;
+				while( SearchClass && !TargetFunc )
+				{
+					// Since the function map is not valid yet, we have to iterate over the fields to look for the function
+					for( UField* TestField = SearchClass->Children; TestField; TestField = TestField->Next )
+					{
+						UFunction* TestFunc = Cast<UFunction>(TestField);
+						if (TestFunc && FNativeClassHeaderGenerator::GetOverriddenFName(TestFunc) == FuncName)
+						{
+							TargetFunc = TestFunc;
+							break;
+						}
+					}
+					SearchClass = SearchClass->GetSuperClass();
+				}
+
+				return TargetFunc;
+			};
+
+			FClassMetaData* TargetClassData = GScriptHelper.FindClassData(TargetClass);
+			check(TargetClassData);
+			FTokenData* PropertyToken = TargetClassData->FindTokenData(Prop);
+			check(PropertyToken);
+
+			TGuardValue<int32> GuardedInputPos(InputPos, PropertyToken->Token.StartPos);
+			TGuardValue<int32> GuardedInputLine(InputLine, PropertyToken->Token.StartLine);
+
+			if (Prop->HasAnyPropertyFlags(CPF_RepNotify))
+			{
+				VerifyRepNotifyCallback(Prop, FindTargetFunction(Prop->RepNotifyFunc));
+			}
+
+			if (Prop->HasAnyPropertyFlags(CPF_BlueprintVisible))
+			{
+				const FString& GetterFuncName = Prop->GetMetaData(NAME_BlueprintGetter);
+				if (!GetterFuncName.IsEmpty())
+				{
+					if (UFunction* TargetFunc = FindTargetFunction(*GetterFuncName))
+					{
+						VerifyBlueprintPropertyGetter(Prop, TargetFunc);
+					}
+					else
+					{
+						// Couldn't find a valid function...
+						UE_LOG_ERROR_UHT(TEXT("Blueprint Property getter function %s not found"), *GetterFuncName);
+					}
+				}
+
+				if (!Prop->HasAnyPropertyFlags(CPF_BlueprintReadOnly))
+				{
+					const FString& SetterFuncName = Prop->GetMetaData(NAME_BlueprintSetter);
+					if (!SetterFuncName.IsEmpty())
+					{
+						if (UFunction* TargetFunc = FindTargetFunction(*SetterFuncName))
+						{
+							VerifyBlueprintPropertySetter(Prop, TargetFunc);
+						}
+						else
+						{
+							// Couldn't find a valid function...
+							UE_LOG_ERROR_UHT(TEXT("Blueprint Property setter function %s not found"), *SetterFuncName);
+						}
+					}
+				}
+			}
+		}
+	}
+}
+
+
+/*-----------------------------------------------------------------------------
+	Compiler directives.
+-----------------------------------------------------------------------------*/
+
+//
+// Process a compiler directive.
+//
+void FHeaderParser::CompileDirective(FClasses& AllClasses)
+{
+	FUnrealSourceFile* CurrentSourceFilePtr = GetCurrentSourceFile();
+	TSharedRef<FUnrealSourceFile> CurrentSrcFile = CurrentSourceFilePtr->AsShared();
+	FToken Directive;
+
+	int32 LineAtStartOfDirective = InputLine;
+	// Define directive are skipped but they can be multiline.
+	bool bDefineDirective = false;
+
+	if (!GetIdentifier(Directive))
+	{
+		FError::Throwf(TEXT("Missing compiler directive after '#'") );
+	}
+	else if (Directive.Matches(TEXT("error"), ESearchCase::CaseSensitive))
+	{
+		FError::Throwf(TEXT("#error directive encountered") );
+	}
+	else if (Directive.Matches(TEXT("pragma"), ESearchCase::CaseSensitive))
+	{
+		// Ignore all pragmas
+	}
+	else if (Directive.Matches(TEXT("linenumber"), ESearchCase::CaseSensitive))
+	{
+		FToken Number;
+		if (!GetToken(Number) || (Number.TokenType != TOKEN_Const) || (Number.Type != CPT_Int && Number.Type != CPT_Int64))
+		{
+			FError::Throwf(TEXT("Missing line number in line number directive"));
+		}
+
+		int32 newInputLine;
+		if ( Number.GetConstInt(newInputLine) )
+		{
+			InputLine = newInputLine;
+		}
+	}
+	else if (Directive.Matches(TEXT("include"), ESearchCase::CaseSensitive))
+	{
+		FToken IncludeName;
+		if (GetToken(IncludeName) && (IncludeName.TokenType == TOKEN_Const) && (IncludeName.Type == CPT_String))
+		{
+			const FString& ExpectedHeaderName = CurrentSrcFile->GetGeneratedHeaderFilename();
+			if (FCString::Stricmp(IncludeName.String, *ExpectedHeaderName) == 0)
+			{
+				bSpottedAutogeneratedHeaderInclude = true;
+			}
+		}
+	}
+	else if (Directive.Matches(TEXT("if"), ESearchCase::CaseSensitive))
+	{
+		// Eat the ! if present
+		bool bNotDefined = MatchSymbol(TEXT('!'));
+
+		int32 TempInt;
+		const bool bParsedInt = GetConstInt(TempInt);
+		if (bParsedInt && (TempInt == 0 || TempInt == 1))
+		{
+			PushCompilerDirective(ECompilerDirective::Insignificant);
+		}
+		else
+		{
+			FToken Define;
+			if (!GetIdentifier(Define))
+			{
+				FError::Throwf(TEXT("Missing define name '#if'") );
+			}
+
+			if ( Define.Matches(TEXT("WITH_EDITORONLY_DATA"), ESearchCase::CaseSensitive) )
+			{
+				PushCompilerDirective(ECompilerDirective::WithEditorOnlyData);
+			}
+			else if ( Define.Matches(TEXT("WITH_EDITOR"), ESearchCase::CaseSensitive) )
+			{
+				PushCompilerDirective(ECompilerDirective::WithEditor);
+			}
+			else if (Define.Matches(TEXT("WITH_HOT_RELOAD"), ESearchCase::CaseSensitive) || Define.Matches(TEXT("WITH_HOT_RELOAD_CTORS"), ESearchCase::CaseSensitive) || Define.Matches(TEXT('1')))
+			{
+				PushCompilerDirective(ECompilerDirective::Insignificant);
+			}
+			else if ( Define.Matches(TEXT("CPP"), ESearchCase::CaseSensitive) && bNotDefined)
+			{
+				PushCompilerDirective(ECompilerDirective::Insignificant);
+			}
+			else
+			{
+				FError::Throwf(TEXT("Unknown define '#if %s' in class or global scope"), Define.Identifier);
+			}
+		}
+	}
+	else if (Directive.Matches(TEXT("endif"), ESearchCase::CaseSensitive))
+	{
+		if (CompilerDirectiveStack.Num() < 1)
+		{
+			FError::Throwf(TEXT("Unmatched '#endif' in class or global scope"));
+		}
+		CompilerDirectiveStack.Pop();
+	}
+	else if (Directive.Matches(TEXT("define"), ESearchCase::CaseSensitive))
+	{
+		// Ignore the define directive (can be multiline).
+		bDefineDirective = true;
+	}
+	else if (Directive.Matches(TEXT("ifdef"), ESearchCase::CaseSensitive) || Directive.Matches(TEXT("ifndef"), ESearchCase::CaseSensitive))
+	{
+		PushCompilerDirective(ECompilerDirective::Insignificant);
+	}
+	else if (Directive.Matches(TEXT("undef"), ESearchCase::CaseSensitive) || Directive.Matches(TEXT("else"), ESearchCase::CaseSensitive))
+	{
+		// Ignore. UHT can only handle #if directive
+	}
+	else
+	{
+		FError::Throwf(TEXT("Unrecognized compiler directive %s"), Directive.Identifier );
+	}
+
+	// Skip to end of line (or end of multiline #define).
+	if (LineAtStartOfDirective == InputLine)
+	{
+		TCHAR LastCharacter = '\0';
+		TCHAR c;
+		do
+		{			
+			while ( !IsEOL( c=GetChar() ) )
+			{
+				LastCharacter = c;
+			}
+		} 
+		// Continue until the entire multiline directive has been skipped.
+		while (LastCharacter == '\\' && bDefineDirective);
+
+		if (c == 0)
+		{
+			UngetChar();
+		}
+	}
+}
+
+/*-----------------------------------------------------------------------------
+	Variable declaration parser.
+-----------------------------------------------------------------------------*/
+
+void FHeaderParser::GetVarType(
+	FClasses&                       AllClasses,
+	FScope*                         Scope,
+	FPropertyBase&                  VarProperty,
+	EPropertyFlags                  Disallow,
+	const FToken*                   OuterPropertyType,
+	EPropertyDeclarationStyle::Type PropertyDeclarationStyle,
+	EVariableCategory::Type         VariableCategory,
+	FIndexRange*                    ParsedVarIndexRange,
+	ELayoutMacroType*               OutLayoutMacroType
+)
+{
+	UStruct* OwnerStruct = Scope->IsFileScope() ? nullptr : ((FStructScope*)Scope)->GetStruct();
+	FName RepCallbackName = FName(NAME_None);
+
+	// Get flags.
+	EPropertyFlags Flags        = CPF_None;
+	EPropertyFlags ImpliedFlags = CPF_None;
+
+	// force members to be 'blueprint read only' if in a const class
+	if (VariableCategory == EVariableCategory::Member)
+	{
+		if (UClass* OwnerClass = Cast<UClass>(OwnerStruct))
+		{
+			if (OwnerClass->ClassFlags & CLASS_Const)
+			{
+				ImpliedFlags |= CPF_BlueprintReadOnly;
+			}
+		}
+	}
+	uint32 ExportFlags = PROPEXPORT_Public;
+
+	// Build up a list of specifiers
+	TArray<FPropertySpecifier> SpecifiersFound;
+
+	TMap<FName, FString> MetaDataFromNewStyle;
+	bool bNativeConst = false;
+	bool bNativeConstTemplateArg = false;
+
+	const bool bIsParamList = (VariableCategory != EVariableCategory::Member) && MatchIdentifier(TEXT("UPARAM"), ESearchCase::CaseSensitive);
+
+	// No specifiers are allowed inside a TArray
+	if ((OuterPropertyType == NULL) || !OuterPropertyType->Matches(TEXT("TArray"), ESearchCase::CaseSensitive))
+	{
+		// New-style UPROPERTY() syntax 
+		if (PropertyDeclarationStyle == EPropertyDeclarationStyle::UPROPERTY || bIsParamList)
+		{
+			ReadSpecifierSetInsideMacro(SpecifiersFound, TEXT("Variable"), MetaDataFromNewStyle);
+		}
+	}
+
+	if (VariableCategory != EVariableCategory::Member)
+	{
+		// const before the variable type support (only for params)
+		if (MatchIdentifier(TEXT("const"), ESearchCase::CaseSensitive))
+		{
+			Flags |= CPF_ConstParm;
+			bNativeConst = true;
+		}
+	}
+
+	if (CompilerDirectiveStack.Num() > 0 && (CompilerDirectiveStack.Last()&ECompilerDirective::WithEditorOnlyData) != 0)
+	{
+		Flags |= CPF_EditorOnly;
+	}
+
+	// Store the start and end positions of the parsed type
+	if (ParsedVarIndexRange)
+	{
+		ParsedVarIndexRange->StartIndex = InputPos;
+	}
+
+	// Process the list of specifiers
+	bool bSeenEditSpecifier = false;
+	bool bSeenBlueprintWriteSpecifier = false;
+	bool bSeenBlueprintReadOnlySpecifier = false;
+	bool bSeenBlueprintGetterSpecifier = false;
+	for (const FPropertySpecifier& Specifier : SpecifiersFound)
+	{
+		EVariableSpecifier SpecID = (EVariableSpecifier)Algo::FindSortedStringCaseInsensitive(*Specifier.Key, GVariableSpecifierStrings);
+		if (VariableCategory == EVariableCategory::Member)
+		{
+			switch (SpecID)
+			{
+				case EVariableSpecifier::EditAnywhere:
+				{
+					if (bSeenEditSpecifier)
+					{
+						UE_LOG_ERROR_UHT(TEXT("Found more than one edit/visibility specifier (%s), only one is allowed"), *Specifier.Key);
+					}
+					Flags |= CPF_Edit;
+					bSeenEditSpecifier = true;
+				}
+				break;
+
+				case EVariableSpecifier::EditInstanceOnly:
+				{
+					if (bSeenEditSpecifier)
+					{
+						UE_LOG_ERROR_UHT(TEXT("Found more than one edit/visibility specifier (%s), only one is allowed"), *Specifier.Key);
+					}
+					Flags |= CPF_Edit | CPF_DisableEditOnTemplate;
+					bSeenEditSpecifier = true;
+				}
+				break;
+
+				case EVariableSpecifier::EditDefaultsOnly:
+				{
+					if (bSeenEditSpecifier)
+					{
+						UE_LOG_ERROR_UHT(TEXT("Found more than one edit/visibility specifier (%s), only one is allowed"), *Specifier.Key);
+					}
+					Flags |= CPF_Edit | CPF_DisableEditOnInstance;
+					bSeenEditSpecifier = true;
+				}
+				break;
+
+				case EVariableSpecifier::VisibleAnywhere:
+				{
+					if (bSeenEditSpecifier)
+					{
+						UE_LOG_ERROR_UHT(TEXT("Found more than one edit/visibility specifier (%s), only one is allowed"), *Specifier.Key);
+					}
+					Flags |= CPF_Edit | CPF_EditConst;
+					bSeenEditSpecifier = true;
+				}
+				break;
+
+				case EVariableSpecifier::VisibleInstanceOnly:
+				{
+					if (bSeenEditSpecifier)
+					{
+						UE_LOG_ERROR_UHT(TEXT("Found more than one edit/visibility specifier (%s), only one is allowed"), *Specifier.Key);
+					}
+					Flags |= CPF_Edit | CPF_EditConst | CPF_DisableEditOnTemplate;
+					bSeenEditSpecifier = true;
+				}
+				break;
+
+				case EVariableSpecifier::VisibleDefaultsOnly:
+				{
+					if (bSeenEditSpecifier)
+					{
+						UE_LOG_ERROR_UHT(TEXT("Found more than one edit/visibility specifier (%s), only one is allowed"), *Specifier.Key);
+					}
+					Flags |= CPF_Edit | CPF_EditConst | CPF_DisableEditOnInstance;
+					bSeenEditSpecifier = true;
+				}
+				break;
+
+				case EVariableSpecifier::BlueprintReadWrite:
+				{
+					if (bSeenBlueprintReadOnlySpecifier)
+					{
+						UE_LOG_ERROR_UHT(TEXT("Cannot specify a property as being both BlueprintReadOnly and BlueprintReadWrite."));
+					}
+
+					const FString* PrivateAccessMD = MetaDataFromNewStyle.Find(NAME_AllowPrivateAccess);  // FBlueprintMetadata::MD_AllowPrivateAccess
+					const bool bAllowPrivateAccess = PrivateAccessMD ? (*PrivateAccessMD != TEXT("false")) : false;
+					if (CurrentAccessSpecifier == ACCESS_Private && !bAllowPrivateAccess)
+					{
+						UE_LOG_ERROR_UHT(TEXT("BlueprintReadWrite should not be used on private members"));
+					}
+
+					if ((Flags & CPF_EditorOnly) != 0 && OwnerStruct->IsA<UScriptStruct>())
+					{
+						UE_LOG_ERROR_UHT(TEXT("Blueprint exposed struct members cannot be editor only"));
+					}
+
+					Flags |= CPF_BlueprintVisible;
+					bSeenBlueprintWriteSpecifier = true;
+				}
+				break;
+
+				case EVariableSpecifier::BlueprintSetter:
+				{
+					if (bSeenBlueprintReadOnlySpecifier)
+					{
+						UE_LOG_ERROR_UHT(TEXT("Cannot specify a property as being both BlueprintReadOnly and having a BlueprintSetter."));
+					}
+
+					if (OwnerStruct->IsA<UScriptStruct>())
+					{
+						UE_LOG_ERROR_UHT(TEXT("Cannot specify BlueprintSetter for a struct member."))
+					}
+
+					const FString BlueprintSetterFunction = RequireExactlyOneSpecifierValue(Specifier);
+					MetaDataFromNewStyle.Add(NAME_BlueprintSetter, BlueprintSetterFunction);
+
+					Flags |= CPF_BlueprintVisible;
+					bSeenBlueprintWriteSpecifier = true;
+				}
+				break;
+
+				case EVariableSpecifier::BlueprintReadOnly:
+				{
+					if (bSeenBlueprintWriteSpecifier)
+					{
+						UE_LOG_ERROR_UHT(TEXT("Cannot specify both BlueprintReadOnly and BlueprintReadWrite or BlueprintSetter."), *Specifier.Key);
+					}
+
+					const FString* PrivateAccessMD = MetaDataFromNewStyle.Find(NAME_AllowPrivateAccess);  // FBlueprintMetadata::MD_AllowPrivateAccess
+					const bool bAllowPrivateAccess = PrivateAccessMD ? (*PrivateAccessMD != TEXT("false")) : false;
+					if (CurrentAccessSpecifier == ACCESS_Private && !bAllowPrivateAccess)
+					{
+						UE_LOG_ERROR_UHT(TEXT("BlueprintReadOnly should not be used on private members"));
+					}
+
+					if ((Flags & CPF_EditorOnly) != 0 && OwnerStruct->IsA<UScriptStruct>())
+					{
+						UE_LOG_ERROR_UHT(TEXT("Blueprint exposed struct members cannot be editor only"));
+					}
+
+					Flags        |= CPF_BlueprintVisible | CPF_BlueprintReadOnly;
+					ImpliedFlags &= ~CPF_BlueprintReadOnly;
+					bSeenBlueprintReadOnlySpecifier = true;
+				}
+				break;
+
+				case EVariableSpecifier::BlueprintGetter:
+				{
+					if (OwnerStruct->IsA<UScriptStruct>())
+					{
+						UE_LOG_ERROR_UHT(TEXT("Cannot specify BlueprintGetter for a struct member."))
+					}
+
+					const FString BlueprintGetterFunction = RequireExactlyOneSpecifierValue(Specifier);
+					MetaDataFromNewStyle.Add(NAME_BlueprintGetter, BlueprintGetterFunction);
+
+					Flags        |= CPF_BlueprintVisible;
+					bSeenBlueprintGetterSpecifier = true;
+				}
+				break;
+
+				case EVariableSpecifier::Config:
+				{
+					Flags |= CPF_Config;
+				}
+				break;
+
+				case EVariableSpecifier::GlobalConfig:
+				{
+					Flags |= CPF_GlobalConfig | CPF_Config;
+				}
+				break;
+
+				case EVariableSpecifier::Localized:
+				{
+					UE_LOG_ERROR_UHT(TEXT("The Localized specifier is deprecated"));
+				}
+				break;
+
+				case EVariableSpecifier::Transient:
+				{
+					Flags |= CPF_Transient;
+				}
+				break;
+
+				case EVariableSpecifier::DuplicateTransient:
+				{
+					Flags |= CPF_DuplicateTransient;
+				}
+				break;
+
+				case EVariableSpecifier::TextExportTransient:
+				{
+					Flags |= CPF_TextExportTransient;
+				}
+				break;
+
+				case EVariableSpecifier::NonPIETransient:
+				{
+					UE_LOG_WARNING_UHT(TEXT("NonPIETransient is deprecated - NonPIEDuplicateTransient should be used instead"));
+					Flags |= CPF_NonPIEDuplicateTransient;
+				}
+				break;
+
+				case EVariableSpecifier::NonPIEDuplicateTransient:
+				{
+					Flags |= CPF_NonPIEDuplicateTransient;
+				}
+				break;
+
+				case EVariableSpecifier::Export:
+				{
+					Flags |= CPF_ExportObject;
+				}
+				break;
+
+				case EVariableSpecifier::EditInline:
+				{
+					UE_LOG_ERROR_UHT(TEXT("EditInline is deprecated. Remove it, or use Instanced instead."));
+				}
+				break;
+
+				case EVariableSpecifier::NoClear:
+				{
+					Flags |= CPF_NoClear;
+				}
+				break;
+
+				case EVariableSpecifier::EditFixedSize:
+				{
+					Flags |= CPF_EditFixedSize;
+				}
+				break;
+
+				case EVariableSpecifier::Replicated:
+				case EVariableSpecifier::ReplicatedUsing:
+				{
+					if (OwnerStruct->IsA<UScriptStruct>())
+					{
+						UE_LOG_ERROR_UHT(TEXT("Struct members cannot be replicated"));
+					}
+
+					Flags |= CPF_Net;
+
+					// See if we've specified a rep notification function
+					if (SpecID == EVariableSpecifier::ReplicatedUsing)
+					{
+						RepCallbackName = FName(*RequireExactlyOneSpecifierValue(Specifier));
+						Flags |= CPF_RepNotify;
+					}
+				}
+				break;
+
+				case EVariableSpecifier::NotReplicated:
+				{
+					if (!OwnerStruct->IsA<UScriptStruct>())
+					{
+						UE_LOG_ERROR_UHT(TEXT("Only Struct members can be marked NotReplicated"));
+					}
+
+					Flags |= CPF_RepSkip;
+				}
+				break;
+
+				case EVariableSpecifier::RepRetry:
+				{
+					UE_LOG_ERROR_UHT(TEXT("'RepRetry' is deprecated."));
+				}
+				break;
+
+				case EVariableSpecifier::Interp:
+				{
+					Flags |= CPF_Edit;
+					Flags |= CPF_BlueprintVisible;
+					Flags |= CPF_Interp;
+				}
+				break;
+
+				case EVariableSpecifier::NonTransactional:
+				{
+					Flags |= CPF_NonTransactional;
+				}
+				break;
+
+				case EVariableSpecifier::Instanced:
+				{
+					Flags |= CPF_PersistentInstance | CPF_ExportObject | CPF_InstancedReference;
+					AddEditInlineMetaData(MetaDataFromNewStyle);
+				}
+				break;
+
+				case EVariableSpecifier::BlueprintAssignable:
+				{
+					Flags |= CPF_BlueprintAssignable;
+				}
+				break;
+
+				case EVariableSpecifier::BlueprintCallable:
+				{
+					Flags |= CPF_BlueprintCallable;
+				}
+				break;
+
+				case EVariableSpecifier::BlueprintAuthorityOnly:
+				{
+					Flags |= CPF_BlueprintAuthorityOnly;
+				}
+				break;
+
+				case EVariableSpecifier::AssetRegistrySearchable:
+				{
+					Flags |= CPF_AssetRegistrySearchable;
+				}
+				break;
+
+				case EVariableSpecifier::SimpleDisplay:
+				{
+					Flags |= CPF_SimpleDisplay;
+				}
+				break;
+
+				case EVariableSpecifier::AdvancedDisplay:
+				{
+					Flags |= CPF_AdvancedDisplay;
+				}
+				break;
+
+				case EVariableSpecifier::SaveGame:
+				{
+					Flags |= CPF_SaveGame;
+				}
+				break;
+
+				case EVariableSpecifier::SkipSerialization:
+				{
+					Flags |= CPF_SkipSerialization;
+				}
+				break;
+
+				default:
+				{
+					UE_LOG_ERROR_UHT(TEXT("Unknown variable specifier '%s'"), *Specifier.Key);
+				}
+				break;
+			}
+		}
+		else
+		{
+			switch (SpecID)
+			{
+				case EVariableSpecifier::Const:
+				{
+					Flags |= CPF_ConstParm;
+				}
+				break;
+
+				case EVariableSpecifier::Ref:
+				{
+					Flags |= CPF_OutParm | CPF_ReferenceParm;
+				}
+				break;
+
+				case EVariableSpecifier::NotReplicated:
+				{
+					if (VariableCategory == EVariableCategory::ReplicatedParameter)
+					{
+						VariableCategory = EVariableCategory::RegularParameter;
+						Flags |= CPF_RepSkip;
+					}
+					else
+					{
+						UE_LOG_ERROR_UHT(TEXT("Only parameters in service request functions can be marked NotReplicated"));
+					}
+				}
+				break;
+
+				default:
+				{
+					UE_LOG_ERROR_UHT(TEXT("Unknown variable specifier '%s'"), *Specifier.Key);
+				}
+				break;
+			}
+		}
+	}
+
+	// If we saw a BlueprintGetter but did not see BlueprintSetter or 
+	// or BlueprintReadWrite then treat as BlueprintReadOnly
+	if (bSeenBlueprintGetterSpecifier && !bSeenBlueprintWriteSpecifier)
+	{
+		Flags |= CPF_BlueprintReadOnly;
+		ImpliedFlags &= ~CPF_BlueprintReadOnly;
+	}
+
+	{
+		const FString* ExposeOnSpawnStr = MetaDataFromNewStyle.Find(NAME_ExposeOnSpawn);
+		const bool bExposeOnSpawn = (NULL != ExposeOnSpawnStr);
+		if (bExposeOnSpawn)
+		{
+			if (0 != (CPF_DisableEditOnInstance & Flags))
+			{
+				UE_LOG_WARNING_UHT(TEXT("Property cannot have both 'DisableEditOnInstance' and 'ExposeOnSpawn' flags"));
+			}
+			if (0 == (CPF_BlueprintVisible & Flags))
+			{
+				UE_LOG_WARNING_UHT(TEXT("Property cannot have 'ExposeOnSpawn' without 'BlueprintVisible' flag."));
+			}
+			Flags |= CPF_ExposeOnSpawn;
+		}
+	}
+
+	if (CurrentAccessSpecifier == ACCESS_Public || VariableCategory != EVariableCategory::Member)
+	{
+		Flags       &= ~CPF_Protected;
+		ExportFlags |= PROPEXPORT_Public;
+		ExportFlags &= ~(PROPEXPORT_Private|PROPEXPORT_Protected);
+
+		Flags &= ~CPF_NativeAccessSpecifiers;
+		Flags |= CPF_NativeAccessSpecifierPublic;
+	}
+	else if (CurrentAccessSpecifier == ACCESS_Protected)
+	{
+		Flags       |= CPF_Protected;
+		ExportFlags |= PROPEXPORT_Protected;
+		ExportFlags &= ~(PROPEXPORT_Public|PROPEXPORT_Private);
+
+		Flags &= ~CPF_NativeAccessSpecifiers;
+		Flags |= CPF_NativeAccessSpecifierProtected;
+	}
+	else if (CurrentAccessSpecifier == ACCESS_Private)
+	{
+		Flags       &= ~CPF_Protected;
+		ExportFlags |= PROPEXPORT_Private;
+		ExportFlags &= ~(PROPEXPORT_Public|PROPEXPORT_Protected);
+
+		Flags &= ~CPF_NativeAccessSpecifiers;
+		Flags |= CPF_NativeAccessSpecifierPrivate;
+	}
+	else
+	{
+		FError::Throwf(TEXT("Unknown access level"));
+	}
+
+	// Swallow inline keywords
+	if (VariableCategory == EVariableCategory::Return)
+	{
+		FToken InlineToken;
+		if (!GetIdentifier(InlineToken, true))
+		{
+			FError::Throwf(TEXT("%s: Missing variable type"), GetHintText(VariableCategory));
+		}
+
+		if (FCString::Strcmp(InlineToken.Identifier, TEXT("inline")) != 0
+			&& FCString::Strcmp(InlineToken.Identifier, TEXT("FORCENOINLINE")) != 0
+			&& FCString::Strncmp(InlineToken.Identifier, TEXT("FORCEINLINE"), 11) != 0)
+		{
+			UngetToken(InlineToken);
+		}
+	}
+
+	// Get variable type.
+	bool bUnconsumedStructKeyword = false;
+	bool bUnconsumedClassKeyword  = false;
+	bool bUnconsumedEnumKeyword   = false;
+	bool bUnconsumedConstKeyword  = false;
+
+	// Handle MemoryLayout.h macros
+	ELayoutMacroType LayoutMacroType     = ELayoutMacroType::None;
+	bool             bHasWrapperBrackets = false;
+	ON_SCOPE_EXIT
+	{
+		if (OutLayoutMacroType)
+		{
+			*OutLayoutMacroType = LayoutMacroType;
+			if (bHasWrapperBrackets)
+			{
+				RequireSymbol(TEXT(')'), GLayoutMacroNames[(int32)LayoutMacroType]);
+			}
+		}
+	};
+
+	if (OutLayoutMacroType)
+	{
+		*OutLayoutMacroType = ELayoutMacroType::None;
+
+		FToken LayoutToken;
+		if (GetToken(LayoutToken))
+		{
+			if (LayoutToken.TokenType == TOKEN_Identifier)
+			{
+				LayoutMacroType = (ELayoutMacroType)Algo::FindSortedStringCaseInsensitive(LayoutToken.Identifier, GLayoutMacroNames, UE_ARRAY_COUNT(GLayoutMacroNames));
+				if (LayoutMacroType != ELayoutMacroType::None)
+				{
+					RequireSymbol(TEXT('('), GLayoutMacroNames[(int32)LayoutMacroType]);
+					if (LayoutMacroType == ELayoutMacroType::ArrayEditorOnly || LayoutMacroType == ELayoutMacroType::FieldEditorOnly || LayoutMacroType == ELayoutMacroType::BitfieldEditorOnly)
+					{
+						Flags |= CPF_EditorOnly;
+					}
+					bHasWrapperBrackets = MatchSymbol(TEXT("("));
+				}
+				else
+				{
+					UngetToken(LayoutToken);
+				}
+			}
+		}
+	}
+
+	if (MatchIdentifier(TEXT("const"), ESearchCase::CaseSensitive))
+	{
+		//@TODO: UCREMOVAL: Should use this to set the new (currently non-existent) CPF_Const flag appropriately!
+		bUnconsumedConstKeyword = true;
+		bNativeConst = true;
+	}
+
+	if (MatchIdentifier(TEXT("mutable"), ESearchCase::CaseSensitive))
+	{
+		//@TODO: Should flag as settable from a const context, but this is at least good enough to allow use for C++ land
+	}
+
+	if (MatchIdentifier(TEXT("struct"), ESearchCase::CaseSensitive))
+	{
+		bUnconsumedStructKeyword = true;
+	}
+	else if (MatchIdentifier(TEXT("class"), ESearchCase::CaseSensitive))
+	{
+		bUnconsumedClassKeyword = true;
+	}
+	else if (MatchIdentifier(TEXT("enum"), ESearchCase::CaseSensitive))
+	{
+		if (VariableCategory == EVariableCategory::Member)
+		{
+			FError::Throwf(TEXT("%s: Cannot declare enum at variable declaration"), GetHintText(VariableCategory));
+		}
+
+		bUnconsumedEnumKeyword = true;
+	}
+
+	//
+	FToken VarType;
+	if ( !GetIdentifier(VarType,1) )
+	{
+		FError::Throwf(TEXT("%s: Missing variable type"), GetHintText(VariableCategory));
+	}
+
+	RedirectTypeIdentifier(VarType);
+
+	if ( VarType.Matches(TEXT("int8"), ESearchCase::CaseSensitive) )
+	{
+		VarProperty = FPropertyBase(CPT_Int8);
+	}
+	else if ( VarType.Matches(TEXT("int16"), ESearchCase::CaseSensitive) )
+	{
+		VarProperty = FPropertyBase(CPT_Int16);
+	}
+	else if ( VarType.Matches(TEXT("int32"), ESearchCase::CaseSensitive) )
+	{
+		VarProperty = FPropertyBase(CPT_Int);
+	}
+	else if ( VarType.Matches(TEXT("int64"), ESearchCase::CaseSensitive) )
+	{
+		VarProperty = FPropertyBase(CPT_Int64);
+	}
+	else if ( VarType.Matches(TEXT("uint64"), ESearchCase::CaseSensitive) && IsBitfieldProperty(LayoutMacroType) )
+	{
+		// 64-bit bitfield (bool) type, treat it like 8 bit type
+		VarProperty = FPropertyBase(CPT_Bool8);
+	}
+	else if ( VarType.Matches(TEXT("uint32"), ESearchCase::CaseSensitive) && IsBitfieldProperty(LayoutMacroType) )
+	{
+		// 32-bit bitfield (bool) type, treat it like 8 bit type
+		VarProperty = FPropertyBase(CPT_Bool8);
+	}
+	else if ( VarType.Matches(TEXT("uint16"), ESearchCase::CaseSensitive) && IsBitfieldProperty(LayoutMacroType) )
+	{
+		// 16-bit bitfield (bool) type, treat it like 8 bit type.
+		VarProperty = FPropertyBase(CPT_Bool8);
+	}
+	else if ( VarType.Matches(TEXT("uint8"), ESearchCase::CaseSensitive) && IsBitfieldProperty(LayoutMacroType) )
+	{
+		// 8-bit bitfield (bool) type
+		VarProperty = FPropertyBase(CPT_Bool8);
+	}
+	else if ( VarType.Matches(TEXT("int"), ESearchCase::CaseSensitive) )
+	{
+		VarProperty = FPropertyBase(CPT_Int, EIntType::Unsized);
+	}
+	else if ( VarType.Matches(TEXT("signed"), ESearchCase::CaseSensitive) )
+	{
+		MatchIdentifier(TEXT("int"), ESearchCase::CaseSensitive);
+		VarProperty = FPropertyBase(CPT_Int, EIntType::Unsized);
+	}
+	else if (VarType.Matches(TEXT("unsigned"), ESearchCase::CaseSensitive))
+	{
+		MatchIdentifier(TEXT("int"), ESearchCase::CaseSensitive);
+		VarProperty = FPropertyBase(CPT_UInt32, EIntType::Unsized);
+	}
+	else if ( VarType.Matches(TEXT("bool"), ESearchCase::CaseSensitive) )
+	{
+		if (IsBitfieldProperty(LayoutMacroType))
+		{
+			UE_LOG_ERROR_UHT(TEXT("bool bitfields are not supported."));
+		}
+		// C++ bool type
+		VarProperty = FPropertyBase(CPT_Bool);
+	}
+	else if ( VarType.Matches(TEXT("uint8"), ESearchCase::CaseSensitive) )
+	{
+		// Intrinsic Byte type.
+		VarProperty = FPropertyBase(CPT_Byte);
+	}
+	else if ( VarType.Matches(TEXT("uint16"), ESearchCase::CaseSensitive) )
+	{
+		VarProperty = FPropertyBase(CPT_UInt16);
+	}
+	else if ( VarType.Matches(TEXT("uint32"), ESearchCase::CaseSensitive) )
+	{
+		VarProperty = FPropertyBase(CPT_UInt32);
+	}
+	else if ( VarType.Matches(TEXT("uint64"), ESearchCase::CaseSensitive) )
+	{
+		VarProperty = FPropertyBase(CPT_UInt64);
+	}
+	else if ( VarType.Matches(TEXT("float"), ESearchCase::CaseSensitive) )
+	{
+		// Intrinsic single precision floating point type.
+		VarProperty = FPropertyBase(CPT_Float);
+	}
+	else if ( VarType.Matches(TEXT("double"), ESearchCase::CaseSensitive) )
+	{
+		// Intrinsic double precision floating point type type.
+		VarProperty = FPropertyBase(CPT_Double);
+	}
+	else if ( VarType.Matches(TEXT("FName"), ESearchCase::CaseSensitive) )
+	{
+		// Intrinsic Name type.
+		VarProperty = FPropertyBase(CPT_Name);
+	}
+	else if ( VarType.Matches(TEXT("TArray"), ESearchCase::CaseSensitive) )
+	{
+		RequireSymbol( TEXT('<'), TEXT("'tarray'") );
+
+		VarType.PropertyFlags = Flags;
+
+		GetVarType(AllClasses, Scope, VarProperty, Disallow, &VarType, EPropertyDeclarationStyle::None, VariableCategory);
+		if (VarProperty.IsContainer())
+		{
+			FError::Throwf(TEXT("Nested containers are not supported.") );
+		}
+		// TODO: Prevent sparse delegate types from being used in a container
+
+		if (VarProperty.MetaData.Find(NAME_NativeConst))
+		{
+			bNativeConstTemplateArg = true;
+		}
+
+		VarType.PropertyFlags = VarProperty.PropertyFlags & (CPF_ContainsInstancedReference | CPF_InstancedReference); // propagate these to the array, we will fix them later
+		VarProperty.ArrayType = EArrayType::Dynamic;
+
+		FToken CloseTemplateToken;
+		if (!GetToken(CloseTemplateToken, /*bNoConsts=*/ true, ESymbolParseOption::CloseTemplateBracket))
+		{
+			FError::Throwf(TEXT("Missing token while parsing TArray."));
+		}
+
+		if (CloseTemplateToken.TokenType != TOKEN_Symbol || !CloseTemplateToken.Matches(TEXT('>')))
+		{
+			// If we didn't find a comma, report it
+			if (!CloseTemplateToken.Matches(TEXT(',')))
+			{
+				FError::Throwf(TEXT("Expected '>' but found '%s'"), CloseTemplateToken.Identifier);
+			}
+
+			// If we found a comma, read the next thing, assume it's an allocator, and report that
+			FToken AllocatorToken;
+			if (!GetToken(AllocatorToken, /*bNoConsts=*/ true, ESymbolParseOption::CloseTemplateBracket))
+			{
+				FError::Throwf(TEXT("Unexpected end of file when parsing TArray allocator."));
+			}
+
+			if (AllocatorToken.TokenType != TOKEN_Identifier)
+			{
+				FError::Throwf(TEXT("Found '%s' - expected a '>' or ','."), AllocatorToken.Identifier);
+			}
+
+			if (FCString::Strcmp(AllocatorToken.Identifier, TEXT("FMemoryImageAllocator")) == 0)
+			{
+				if (EnumHasAnyFlags(Flags, CPF_Net))
+				{
+					FError::Throwf(TEXT("Replicated arrays with MemoryImageAllocators are not yet supported"));
+				}
+
+				RequireSymbol(TEXT('>'), TEXT("TArray template arguments"), ESymbolParseOption::CloseTemplateBracket);
+
+				VarProperty.AllocatorType = EAllocatorType::MemoryImage;
+			}
+			else if (FCString::Strcmp(AllocatorToken.Identifier, TEXT("TMemoryImageAllocator")) == 0)
+			{
+				if (EnumHasAnyFlags(Flags, CPF_Net))
+				{
+					FError::Throwf(TEXT("Replicated arrays with MemoryImageAllocators are not yet supported"));
+				}
+
+				RequireSymbol(TEXT('<'), TEXT("TMemoryImageAllocator template arguments"));
+
+				FToken SkipToken;
+				for (;;)
+				{
+					if (!GetToken(SkipToken, /*bNoConsts=*/ false, ESymbolParseOption::CloseTemplateBracket))
+					{
+						FError::Throwf(TEXT("Unexpected end of file when parsing TMemoryImageAllocator template arguments."));
+					}
+
+					if (SkipToken.TokenType == TOKEN_Symbol && FCString::Strcmp(SkipToken.Identifier, TEXT(">")) == 0)
+					{
+						RequireSymbol(TEXT('>'), TEXT("TArray template arguments"), ESymbolParseOption::CloseTemplateBracket);
+						VarProperty.AllocatorType = EAllocatorType::MemoryImage;
+						break;
+					}
+				}
+			}
+			else
+			{
+				FError::Throwf(TEXT("Found '%s' - explicit allocators are not supported in TArray properties."), AllocatorToken.Identifier);
+			}
+		}
+	}
+	else if ( VarType.Matches(TEXT("TMap"), ESearchCase::CaseSensitive) )
+	{
+		RequireSymbol( TEXT('<'), TEXT("'tmap'") );
+
+		VarType.PropertyFlags = Flags;
+
+		FToken MapKeyType;
+		GetVarType(AllClasses, Scope, MapKeyType, Disallow, &VarType, EPropertyDeclarationStyle::None, VariableCategory);
+		if (MapKeyType.IsContainer())
+		{
+			FError::Throwf(TEXT("Nested containers are not supported.") );
+		}
+		// TODO: Prevent sparse delegate types from being used in a container
+
+		if (MapKeyType.Type == CPT_Interface)
+		{
+			FError::Throwf(TEXT("UINTERFACEs are not currently supported as key types."));
+		}
+
+		if (MapKeyType.Type == CPT_Text)
+		{
+			FError::Throwf(TEXT("FText is not currently supported as a key type."));
+		}
+
+		FToken CommaToken;
+		if (!GetToken(CommaToken, /*bNoConsts=*/ true) || CommaToken.TokenType != TOKEN_Symbol || !CommaToken.Matches(TEXT(',')))
+		{
+			FError::Throwf(TEXT("Missing value type while parsing TMap."));
+		}
+
+		GetVarType(AllClasses, Scope, VarProperty, Disallow, &VarType, EPropertyDeclarationStyle::None, VariableCategory);
+		if (VarProperty.IsContainer())
+		{
+			FError::Throwf(TEXT("Nested containers are not supported.") );
+		}
+		// TODO: Prevent sparse delegate types from being used in a container
+
+		EPropertyFlags InnerFlags = (MapKeyType.PropertyFlags | VarProperty.PropertyFlags) & (CPF_ContainsInstancedReference | CPF_InstancedReference); // propagate these to the map value, we will fix them later
+		VarType.PropertyFlags = InnerFlags;
+		VarProperty.MapKeyProp = MakeShared<FToken>(MoveTemp(MapKeyType));
+		VarProperty.MapKeyProp->PropertyFlags = InnerFlags | (VarProperty.MapKeyProp->PropertyFlags & CPF_UObjectWrapper); // Make sure the 'UObjectWrapper' flag is maintained so that 'TMap<TSubclassOf<...>, ...>' works
+
+		FToken CloseTemplateToken;
+		if (!GetToken(CloseTemplateToken, /*bNoConsts=*/ true, ESymbolParseOption::CloseTemplateBracket))
+		{
+			FError::Throwf(TEXT("Missing token while parsing TMap."));
+		}
+
+		if (CloseTemplateToken.TokenType != TOKEN_Symbol || !CloseTemplateToken.Matches(TEXT('>')))
+		{
+			// If we didn't find a comma, report it
+			if (!CloseTemplateToken.Matches(TEXT(',')))
+			{
+				FError::Throwf(TEXT("Expected '>' but found '%s'"), CloseTemplateToken.Identifier);
+			}
+
+			// If we found a comma, read the next thing, assume it's an allocator, and report that
+			FToken AllocatorToken;
+			if (!GetToken(AllocatorToken, /*bNoConsts=*/ true, ESymbolParseOption::CloseTemplateBracket))
+			{
+				FError::Throwf(TEXT("Unexpected end of file when parsing TArray allocator."));
+			}
+
+			if (AllocatorToken.TokenType != TOKEN_Identifier)
+			{
+				FError::Throwf(TEXT("Found '%s' - expected a '>' or ','."), AllocatorToken.Identifier);
+			}
+
+			if (FCString::Strcmp(AllocatorToken.Identifier, TEXT("FMemoryImageSetAllocator")) == 0)
+			{
+				if (EnumHasAnyFlags(Flags, CPF_Net))
+				{
+					FError::Throwf(TEXT("Replicated maps with MemoryImageSetAllocators are not yet supported"));
+				}
+
+				RequireSymbol(TEXT('>'), TEXT("TMap template arguments"), ESymbolParseOption::CloseTemplateBracket);
+
+				VarProperty.AllocatorType = EAllocatorType::MemoryImage;
+			}
+			else
+			{
+				FError::Throwf(TEXT("Found '%s' - explicit allocators are not supported in TMap properties."), AllocatorToken.Identifier);
+			}
+		}
+	}
+	else if ( VarType.Matches(TEXT("TSet"), ESearchCase::CaseSensitive) )
+	{
+		RequireSymbol( TEXT('<'), TEXT("'tset'") );
+
+		VarType.PropertyFlags = Flags;
+
+		GetVarType(AllClasses, Scope, VarProperty, Disallow, &VarType, EPropertyDeclarationStyle::None, VariableCategory);
+		if (VarProperty.IsContainer())
+		{
+			FError::Throwf(TEXT("Nested containers are not supported.") );
+		}
+		// TODO: Prevent sparse delegate types from being used in a container
+
+		if (VarProperty.Type == CPT_Interface)
+		{
+			FError::Throwf(TEXT("UINTERFACEs are not currently supported as element types."));
+		}
+
+		if (VarProperty.Type == CPT_Text)
+		{
+			FError::Throwf(TEXT("FText is not currently supported as an element type."));
+		}
+
+		VarType.PropertyFlags = VarProperty.PropertyFlags & (CPF_ContainsInstancedReference | CPF_InstancedReference); // propagate these to the set, we will fix them later
+		VarProperty.ArrayType = EArrayType::Set;
+
+		FToken CloseTemplateToken;
+		if (!GetToken(CloseTemplateToken, /*bNoConsts=*/ true, ESymbolParseOption::CloseTemplateBracket))
+		{
+			FError::Throwf(TEXT("Missing token while parsing TArray."));
+		}
+
+		if (CloseTemplateToken.TokenType != TOKEN_Symbol || !CloseTemplateToken.Matches(TEXT('>')))
+		{
+			// If we didn't find a comma, report it
+			if (!CloseTemplateToken.Matches(TEXT(',')))
+			{
+				FError::Throwf(TEXT("Expected '>' but found '%s'"), CloseTemplateToken.Identifier);
+			}
+
+			// If we found a comma, read the next thing, assume it's a keyfuncs, and report that
+			FToken AllocatorToken;
+			if (!GetToken(AllocatorToken, /*bNoConsts=*/ true, ESymbolParseOption::CloseTemplateBracket))
+			{
+				FError::Throwf(TEXT("Expected '>' but found '%s'"), CloseTemplateToken.Identifier);
+			}
+
+			FError::Throwf(TEXT("Found '%s' - explicit KeyFuncs are not supported in TSet properties."), AllocatorToken.Identifier);
+		}
+	}
+	else if ( VarType.Matches(TEXT("FString"), ESearchCase::CaseSensitive) || VarType.Matches(TEXT("FMemoryImageString"), ESearchCase::CaseSensitive))
+	{
+		VarProperty = FPropertyBase(CPT_String);
+
+		if (VariableCategory != EVariableCategory::Member)
+		{
+			if (MatchSymbol(TEXT('&')))
+			{
+				if (Flags & CPF_ConstParm)
+				{
+					// 'const FString& Foo' came from 'FString' in .uc, no flags
+					Flags &= ~CPF_ConstParm;
+
+					// We record here that we encountered a const reference, because we need to remove that information from flags for code generation purposes.
+					VarProperty.RefQualifier = ERefQualifier::ConstRef;
+				}
+				else
+				{
+					// 'FString& Foo' came from 'out FString' in .uc
+					Flags |= CPF_OutParm;
+
+					// And we record here that we encountered a non-const reference here too.
+					VarProperty.RefQualifier = ERefQualifier::NonConstRef;
+				}
+			}
+		}
+	}
+	else if ( VarType.Matches(TEXT("Text"), ESearchCase::IgnoreCase) )
+	{
+		FError::Throwf(TEXT("%s' is missing a prefix, expecting 'FText'"), VarType.Identifier);
+	}
+	else if ( VarType.Matches(TEXT("FText"), ESearchCase::CaseSensitive) )
+	{
+		VarProperty = FPropertyBase(CPT_Text);
+	}
+	else if (VarType.Matches(TEXT("TEnumAsByte"), ESearchCase::CaseSensitive))
+	{
+		RequireSymbol(TEXT('<'), VarType.Identifier);
+
+		// Eat the forward declaration enum text if present
+		MatchIdentifier(TEXT("enum"), ESearchCase::CaseSensitive);
+
+		bool bFoundEnum = false;
+
+		FToken InnerEnumType;
+		if (GetIdentifier(InnerEnumType, true))
+		{
+			if (UEnum* Enum = FindObject<UEnum>(ANY_PACKAGE, InnerEnumType.Identifier))
+			{
+				// In-scope enumeration.
+				VarProperty = FPropertyBase(Enum, CPT_Byte);
+				bFoundEnum  = true;
+			}
+		}
+
+		// Try to handle namespaced enums
+		// Note: We do not verify the scoped part is correct, and trust in the C++ compiler to catch that sort of mistake
+		if (MatchSymbol(TEXT("::")))
+		{
+			FToken ScopedTrueEnumName;
+			if (!GetIdentifier(ScopedTrueEnumName, true))
+			{
+				FError::Throwf(TEXT("Expected a namespace scoped enum name.") );
+			}
+		}
+
+		if (!bFoundEnum)
+		{
+			FError::Throwf(TEXT("Expected the name of a previously defined enum"));
+		}
+
+		RequireSymbol(TEXT('>'), VarType.Identifier, ESymbolParseOption::CloseTemplateBracket);
+	}
+	else if (VarType.Matches(TEXT("TFieldPath"), ESearchCase::CaseSensitive ))
+	{
+		RequireSymbol( TEXT('<'), TEXT("'TFieldPath'") );
+
+		FFieldClass* PropertyClass = nullptr;
+		FToken PropertyTypeToken;
+		if (!GetToken(PropertyTypeToken, /* bNoConsts = */ true))
+		{
+			FError::Throwf(TEXT("Expected the property type"));
+		}
+		else
+		{
+			FFieldClass** PropertyClassPtr = FFieldClass::GetNameToFieldClassMap().Find(PropertyTypeToken.Identifier + 1);
+			if (PropertyClassPtr)
+			{
+				PropertyClass = *PropertyClassPtr;
+			}
+			else
+			{
+				FError::Throwf(TEXT("Undefined property type: %s"), PropertyTypeToken.Identifier);
+			}
+		}
+
+		RequireSymbol(TEXT('>'), VarType.Identifier, ESymbolParseOption::CloseTemplateBracket);
+
+		VarProperty = FPropertyBase(PropertyClass, CPT_FieldPath);
+	}
+	else if (UEnum* Enum = FindObject<UEnum>(ANY_PACKAGE, VarType.Identifier))
+	{
+		EPropertyType UnderlyingType = CPT_Byte;
+
+		if (VariableCategory == EVariableCategory::Member)
+		{
+			EUnderlyingEnumType* EnumUnderlyingType = GEnumUnderlyingTypes.Find(Enum);
+			if (!EnumUnderlyingType)
+			{
+				FError::Throwf(TEXT("You cannot use the raw enum name as a type for member variables, instead use TEnumAsByte or a C++11 enum class with an explicit underlying type."), *Enum->CppType);
+			}
+		}
+
+		// Try to handle namespaced enums
+		// Note: We do not verify the scoped part is correct, and trust in the C++ compiler to catch that sort of mistake
+		if (MatchSymbol(TEXT("::")))
+		{
+			FToken ScopedTrueEnumName;
+			if (!GetIdentifier(ScopedTrueEnumName, true))
+			{
+				FError::Throwf(TEXT("Expected a namespace scoped enum name.") );
+			}
+		}
+
+		// In-scope enumeration.
+		VarProperty            = FPropertyBase(Enum, UnderlyingType);
+		bUnconsumedEnumKeyword = false;
+	}
+	else
+	{
+		// Check for structs/classes
+		bool bHandledType = false;
+		FString IdentifierStripped = GetClassNameWithPrefixRemoved(VarType.Identifier);
+		bool bStripped = false;
+		UScriptStruct* Struct = FindObject<UScriptStruct>( ANY_PACKAGE, VarType.Identifier );
+		if (!Struct)
+		{
+			Struct = FindObject<UScriptStruct>( ANY_PACKAGE, *IdentifierStripped );
+			bStripped = true;
+		}
+
+		auto SetDelegateType = [&](UFunction* InFunction, const FString& InIdentifierStripped)
+		{
+			bHandledType = true;
+
+			VarProperty = FPropertyBase(InFunction->HasAnyFunctionFlags(FUNC_MulticastDelegate) ? CPT_MulticastDelegate : CPT_Delegate);
+			VarProperty.DelegateName = *InIdentifierStripped;
+			VarProperty.Function = InFunction;
+
+			if (!(Disallow & CPF_InstancedReference))
+			{
+				Flags |= CPF_InstancedReference;
+			}
+		};
+
+		if (!Struct && MatchSymbol(TEXT("::")))
+		{
+			FToken DelegateName;
+			if (GetIdentifier(DelegateName))
+			{
+				UClass* LocalOwnerClass = AllClasses.FindClass(*IdentifierStripped);
+				if (LocalOwnerClass)
+				{
+					TSharedRef<FScope> LocScope = FScope::GetTypeScope(LocalOwnerClass);
+					const FString DelegateIdentifierStripped = GetClassNameWithPrefixRemoved(DelegateName.Identifier);
+					if (UFunction* DelegateFunc = Cast<UFunction>(LocScope->FindTypeByName(*(DelegateIdentifierStripped + HEADER_GENERATED_DELEGATE_SIGNATURE_SUFFIX))))
+					{
+						SetDelegateType(DelegateFunc, DelegateIdentifierStripped);
+						VarProperty.DelegateSignatureOwnerClass = LocalOwnerClass;
+					}
+				}
+				else
+				{
+					FError::Throwf(TEXT("Cannot find class '%s', to resolve delegate '%s'"), *IdentifierStripped, DelegateName.Identifier);
+				}
+			}
+		}
+
+		if (bHandledType)
+		{
+		}
+		else if (Struct)
+		{
+			if (bStripped)
+			{
+				const TCHAR* PrefixCPP = StructsWithTPrefix.Contains(IdentifierStripped) ? TEXT("T") : Struct->GetPrefixCPP();
+				FString ExpectedStructName = FString::Printf(TEXT("%s%s"), PrefixCPP, *Struct->GetName() );
+				if( FString(VarType.Identifier) != ExpectedStructName )
+				{
+					FError::Throwf( TEXT("Struct '%s' is missing or has an incorrect prefix, expecting '%s'"), VarType.Identifier, *ExpectedStructName );
+				}
+			}
+			else if( !StructsWithNoPrefix.Contains(VarType.Identifier) )
+			{
+				const TCHAR* PrefixCPP = StructsWithTPrefix.Contains(VarType.Identifier) ? TEXT("T") : Struct->GetPrefixCPP();
+				FError::Throwf(TEXT("Struct '%s' is missing a prefix, expecting '%s'"), VarType.Identifier, *FString::Printf(TEXT("%s%s"), PrefixCPP, *Struct->GetName()) );
+			}
+
+			bHandledType = true;
+
+			VarProperty = FPropertyBase( Struct );
+			if((Struct->StructFlags & STRUCT_HasInstancedReference) && !(Disallow & CPF_ContainsInstancedReference))
+			{
+				Flags |= CPF_ContainsInstancedReference;
+			}
+			// Struct keyword in front of a struct is legal, we 'consume' it
+			bUnconsumedStructKeyword = false;
+		}
+		else if ( FindObject<UScriptStruct>( ANY_PACKAGE, *IdentifierStripped ) != nullptr)
+		{
+			bHandledType = true;
+
+			// Struct keyword in front of a struct is legal, we 'consume' it
+			bUnconsumedStructKeyword = false;
+		}
+		else if (UFunction* DelegateFunc = Cast<UFunction>(Scope->FindTypeByName(*(IdentifierStripped + HEADER_GENERATED_DELEGATE_SIGNATURE_SUFFIX))))
+		{
+			SetDelegateType(DelegateFunc, IdentifierStripped);
+		}
+		else
+		{
+			// An object reference of some type (maybe a restricted class?)
+			UClass* TempClass = NULL;
+
+			const bool bIsLazyPtrTemplate        = VarType.Matches(TEXT("TLazyObjectPtr"), ESearchCase::CaseSensitive);
+			const bool bIsSoftObjectPtrTemplate  = VarType.Matches(TEXT("TSoftObjectPtr"), ESearchCase::CaseSensitive);
+			const bool bIsSoftClassPtrTemplate   = VarType.Matches(TEXT("TSoftClassPtr"), ESearchCase::CaseSensitive);
+			const bool bIsWeakPtrTemplate        = VarType.Matches(TEXT("TWeakObjectPtr"), ESearchCase::CaseSensitive);
+			const bool bIsAutoweakPtrTemplate    = VarType.Matches(TEXT("TAutoWeakObjectPtr"), ESearchCase::CaseSensitive);
+			const bool bIsScriptInterfaceWrapper = VarType.Matches(TEXT("TScriptInterface"), ESearchCase::CaseSensitive);
+
+			bool bIsWeak     = false;
+			bool bIsLazy     = false;
+			bool bIsSoft     = false;
+			bool bWeakIsAuto = false;
+
+			if (VarType.Matches(TEXT("TSubclassOf"), ESearchCase::CaseSensitive))
+			{
+				TempClass = UClass::StaticClass();
+			}
+			else if (VarType.Matches(TEXT("FScriptInterface"), ESearchCase::CaseSensitive))
+			{
+				TempClass = UInterface::StaticClass();
+				Flags |= CPF_UObjectWrapper;
+			}
+			else if (bIsSoftClassPtrTemplate)
+			{
+				TempClass = UClass::StaticClass();
+				bIsSoft = true;
+			}
+			else if (bIsLazyPtrTemplate || bIsWeakPtrTemplate || bIsAutoweakPtrTemplate || bIsScriptInterfaceWrapper || bIsSoftObjectPtrTemplate)
+			{
+				RequireSymbol(TEXT('<'), VarType.Identifier);
+
+				// Consume a forward class declaration 'class' if present
+				MatchIdentifier(TEXT("class"), ESearchCase::CaseSensitive);
+
+				// Also consume const
+				bNativeConstTemplateArg |= MatchIdentifier(TEXT("const"), ESearchCase::CaseSensitive);
+				
+				// Find the lazy/weak class
+				FToken InnerClass;
+				if (GetIdentifier(InnerClass))
+				{
+					RedirectTypeIdentifier(InnerClass);
+
+					TempClass = AllClasses.FindScriptClass(InnerClass.Identifier);
+					if (TempClass == nullptr)
+					{
+						FError::Throwf(TEXT("Unrecognized type '%s' (in expression %s<%s>) - type must be a UCLASS"), InnerClass.Identifier, VarType.Identifier, InnerClass.Identifier);
+					}
+
+					if (bIsAutoweakPtrTemplate)
+					{
+						bIsWeak = true;
+						bWeakIsAuto = true;
+					}
+					else if (bIsLazyPtrTemplate)
+					{
+						bIsLazy = true;
+					}
+					else if (bIsWeakPtrTemplate)
+					{
+						bIsWeak = true;
+					}
+					else if (bIsSoftObjectPtrTemplate)
+					{
+						bIsSoft = true;
+					}
+
+					Flags |= CPF_UObjectWrapper;
+				}
+				else
+				{
+					FError::Throwf(TEXT("%s: Missing template type"), VarType.Identifier);
+				}
+
+				RequireSymbol(TEXT('>'), VarType.Identifier, ESymbolParseOption::CloseTemplateBracket);
+			}
+			else
+			{
+				TempClass = AllClasses.FindScriptClass(VarType.Identifier);
+			}
+
+			if (TempClass != NULL)
+			{
+				bHandledType = true;
+
+				bool bAllowWeak = !(Disallow & CPF_AutoWeak); // if it is not allowing anything, force it strong. this is probably a function arg
+				VarProperty = FPropertyBase(TempClass, bAllowWeak && bIsWeak, bWeakIsAuto, bIsLazy, bIsSoft);
+				if (TempClass->IsChildOf(UClass::StaticClass()))
+				{
+					if ( MatchSymbol(TEXT('<')) )
+					{
+						Flags |= CPF_UObjectWrapper;
+
+						// Consume a forward class declaration 'class' if present
+						MatchIdentifier(TEXT("class"), ESearchCase::CaseSensitive);
+
+						// Get the actual class type to restrict this to
+						FToken Limitor;
+						if( !GetIdentifier(Limitor) )
+						{
+							FError::Throwf(TEXT("'class': Missing class limitor"));
+						}
+
+						RedirectTypeIdentifier(Limitor);
+
+						VarProperty.MetaClass = AllClasses.FindScriptClassOrThrow(Limitor.Identifier);
+
+						RequireSymbol( TEXT('>'), TEXT("'class limitor'"), ESymbolParseOption::CloseTemplateBracket );
+					}
+					else
+					{
+						VarProperty.MetaClass = UObject::StaticClass();
+					}
+
+					if (bIsWeak)
+					{
+						FError::Throwf(TEXT("Class variables cannot be weak, they are always strong."));
+					}
+
+					if (bIsLazy)
+					{
+						FError::Throwf(TEXT("Class variables cannot be lazy, they are always strong."));
+					}
+
+					if (bIsSoftObjectPtrTemplate)
+					{
+						FError::Throwf(TEXT("Class variables cannot be stored in TSoftObjectPtr, use TSoftClassPtr instead."));
+					}
+				}
+
+				// Inherit instancing flags
+				if (DoesAnythingInHierarchyHaveDefaultToInstanced(TempClass))
+				{
+					Flags |= ((CPF_InstancedReference|CPF_ExportObject) & (~Disallow)); 
+				}
+
+				// Eat the star that indicates this is a pointer to the UObject
+				if (!(Flags & CPF_UObjectWrapper))
+				{
+					// Const after variable type but before pointer symbol
+					bNativeConst |= MatchIdentifier(TEXT("const"), ESearchCase::CaseSensitive);
+
+					RequireSymbol(TEXT('*'), TEXT("Expected a pointer type"));
+
+					// Swallow trailing 'const' after pointer properties
+					if (VariableCategory == EVariableCategory::Member)
+					{
+						MatchIdentifier(TEXT("const"), ESearchCase::CaseSensitive);
+					}
+
+					VarProperty.PointerType = EPointerType::Native;
+				}
+
+				// Imply const if it's a parameter that is a pointer to a const class
+				if (VariableCategory != EVariableCategory::Member && (TempClass != NULL) && (TempClass->HasAnyClassFlags(CLASS_Const)))
+				{
+					Flags |= CPF_ConstParm;
+				}
+
+				// Class keyword in front of a class is legal, we 'consume' it
+				bUnconsumedClassKeyword = false;
+				bUnconsumedConstKeyword = false;
+			}
+		}
+
+		// Resolve delegates declared in another class  //@TODO: UCREMOVAL: This seems extreme
+		if (!bHandledType)
+		{
+			if (UFunction* DelegateFunc = (UFunction*)StaticFindObject(UFunction::StaticClass(), ANY_PACKAGE, *(IdentifierStripped + HEADER_GENERATED_DELEGATE_SIGNATURE_SUFFIX)))
+			{
+				SetDelegateType(DelegateFunc, IdentifierStripped);
+			}
+
+			if (!bHandledType)
+			{
+				FError::Throwf(TEXT("Unrecognized type '%s' - type must be a UCLASS, USTRUCT or UENUM"), VarType.Identifier );
+			}
+		}
+	}
+
+	if (VariableCategory != EVariableCategory::Member)
+	{
+		// const after the variable type support (only for params)
+		if (MatchIdentifier(TEXT("const"), ESearchCase::CaseSensitive))
+		{
+			Flags |= CPF_ConstParm;
+			bNativeConst = true;
+		}
+	}
+
+	if (bUnconsumedConstKeyword)
+	{
+		if (VariableCategory == EVariableCategory::Member)
+		{
+			FError::Throwf(TEXT("Const properties are not supported."));
+		}
+		else
+		{
+			FError::Throwf(TEXT("Inappropriate keyword 'const' on variable of type '%s'"), VarType.Identifier);
+		}
+	}
+
+	if (bUnconsumedClassKeyword)
+	{
+		FError::Throwf(TEXT("Inappropriate keyword 'class' on variable of type '%s'"), VarType.Identifier );
+	}
+
+	if (bUnconsumedStructKeyword)
+	{
+		FError::Throwf(TEXT("Inappropriate keyword 'struct' on variable of type '%s'"), VarType.Identifier );
+	}
+
+	if (bUnconsumedEnumKeyword)
+	{
+		FError::Throwf(TEXT("Inappropriate keyword 'enum' on variable of type '%s'"), VarType.Identifier );
+	}
+
+	if (MatchSymbol(TEXT('*')))
+	{
+		FError::Throwf(TEXT("Inappropriate '*' on variable of type '%s', cannot have an exposed pointer to this type."), VarType.Identifier );
+	}
+
+	//@TODO: UCREMOVAL: 'const' member variables that will get written post-construction by defaultproperties
+	if (VariableCategory == EVariableCategory::Member && OwnerStruct->IsA<UClass>() && ((UClass*)OwnerStruct)->HasAnyClassFlags(CLASS_Const))
+	{
+		// Eat a 'not quite truthful' const after the type; autogenerated for member variables of const classes.
+		bNativeConst |= MatchIdentifier(TEXT("const"), ESearchCase::CaseSensitive);
+	}
+
+	// Arrays are passed by reference but are only implicitly so; setting it explicitly could cause a problem with replicated functions
+	if (MatchSymbol(TEXT('&')))
+	{
+		switch (VariableCategory)
+		{
+			case EVariableCategory::RegularParameter:
+			case EVariableCategory::Return:
+			{
+				Flags |= CPF_OutParm;
+
+				//@TODO: UCREMOVAL: How to determine if we have a ref param?
+				if (Flags & CPF_ConstParm)
+				{
+					Flags |= CPF_ReferenceParm;
+				}
+			}
+			break;
+
+			case EVariableCategory::ReplicatedParameter:
+			{
+				if (!(Flags & CPF_ConstParm))
+				{
+					FError::Throwf(TEXT("Replicated %s parameters cannot be passed by non-const reference"), VarType.Identifier);
+				}
+
+				Flags |= CPF_ReferenceParm;
+			}
+			break;
+
+			default:
+			{
+			}
+			break;
+		}
+
+		if (Flags & CPF_ConstParm)
+		{
+			VarProperty.RefQualifier = ERefQualifier::ConstRef;
+		}
+		else
+		{
+			VarProperty.RefQualifier = ERefQualifier::NonConstRef;
+		}
+	}
+
+	VarProperty.PropertyExportFlags = ExportFlags;
+
+	// Set FPropertyBase info.
+	VarProperty.PropertyFlags        |= Flags | ImpliedFlags;
+	VarProperty.ImpliedPropertyFlags |= ImpliedFlags;
+
+	// Set the RepNotify name, if the variable needs it
+	if( VarProperty.PropertyFlags & CPF_RepNotify )
+	{
+		if( RepCallbackName != NAME_None )
+		{
+			VarProperty.RepNotifyName = RepCallbackName;
+		}
+		else
+		{
+			FError::Throwf(TEXT("Must specify a valid function name for replication notifications"));
+		}
+	}
+
+	// Perform some more specific validation on the property flags
+	if (VarProperty.PropertyFlags & CPF_PersistentInstance)
+	{
+		if (VarProperty.Type == CPT_ObjectReference)
+		{
+			if (VarProperty.PropertyClass->IsChildOf<UClass>())
+			{
+				FError::Throwf(TEXT("'Instanced' cannot be applied to class properties (UClass* or TSubclassOf<>)"));
+			}
+		}
+		else
+		{
+			FError::Throwf(TEXT("'Instanced' is only allowed on an object property, an array of objects, a set of objects, or a map with an object value type."));
+		}
+	}
+
+	if ( VarProperty.IsObject() && VarProperty.Type != CPT_SoftObjectReference && VarProperty.MetaClass == nullptr && (VarProperty.PropertyFlags&CPF_Config) != 0 )
+	{
+		FError::Throwf(TEXT("Not allowed to use 'config' with object variables"));
+	}
+
+	if ((VarProperty.PropertyFlags & CPF_BlueprintAssignable) && VarProperty.Type != CPT_MulticastDelegate)
+	{
+		FError::Throwf(TEXT("'BlueprintAssignable' is only allowed on multicast delegate properties"));
+	}
+
+	if ((VarProperty.PropertyFlags & CPF_BlueprintCallable) && VarProperty.Type != CPT_MulticastDelegate)
+	{
+		FError::Throwf(TEXT("'BlueprintCallable' is only allowed on a property when it is a multicast delegate"));
+	}
+
+	if ((VarProperty.PropertyFlags & CPF_BlueprintAuthorityOnly) && VarProperty.Type != CPT_MulticastDelegate)
+	{
+		FError::Throwf(TEXT("'BlueprintAuthorityOnly' is only allowed on a property when it is a multicast delegate"));
+	}
+	
+	if (VariableCategory != EVariableCategory::Member)
+	{
+		// These conditions are checked externally for struct/member variables where the flag can be inferred later on from the variable name itself
+		ValidatePropertyIsDeprecatedIfNecessary(VarProperty, OuterPropertyType);
+	}
+
+	// Check for invalid transients
+	EPropertyFlags Transients = VarProperty.PropertyFlags & (CPF_DuplicateTransient | CPF_TextExportTransient | CPF_NonPIEDuplicateTransient);
+	if (Transients && !Cast<UClass>(OwnerStruct))
+	{
+		TArray<const TCHAR*> FlagStrs = ParsePropertyFlags(Transients);
+		FError::Throwf(TEXT("'%s' specifier(s) are only allowed on class member variables"), *FString::Join(FlagStrs, TEXT(", ")));
+	}
+
+	// Make sure the overrides are allowed here.
+	if( VarProperty.PropertyFlags & Disallow )
+	{
+		FError::Throwf(TEXT("Specified type modifiers not allowed here") );
+	}
+
+	// For now, copy the flags that a TMap value has to the key
+	if (FPropertyBase* KeyProp = VarProperty.MapKeyProp.Get())
+	{
+		// Make sure the 'UObjectWrapper' flag is maintained so that both 'TMap<TSubclassOf<...>, ...>' and 'TMap<UClass*, TSubclassOf<...>>' works correctly
+		KeyProp->PropertyFlags = (VarProperty.PropertyFlags & ~CPF_UObjectWrapper) | (KeyProp->PropertyFlags & CPF_UObjectWrapper);
+	}
+
+	VarProperty.MetaData = MetaDataFromNewStyle;
+	if (bNativeConst)
+	{
+		VarProperty.MetaData.Add(NAME_NativeConst, FString());
+	}
+	if (bNativeConstTemplateArg)
+	{
+		VarProperty.MetaData.Add(NAME_NativeConstTemplateArg, FString());
+	}
+	
+	if (ParsedVarIndexRange)
+	{
+		ParsedVarIndexRange->Count = InputPos - ParsedVarIndexRange->StartIndex;
+	}
+}
+
+/**
+ * If the property has already been seen during compilation, then return add. If not,
+ * then return replace so that INI files don't mess with header exporting
+ *
+ * @param PropertyName the string token for the property
+ *
+ * @return FNAME_Replace_Not_Safe_For_Threading or FNAME_Add
+ */
+EFindName FHeaderParser::GetFindFlagForPropertyName(const TCHAR* PropertyName)
+{
+	static TMap<FString,int32> PreviousNames;
+	FString PropertyStr(PropertyName);
+	FString UpperPropertyStr = PropertyStr.ToUpper();
+	// See if it's in the list already
+	if (PreviousNames.Find(UpperPropertyStr))
+	{
+		return FNAME_Add;
+	}
+	// Add it to the list for future look ups
+	PreviousNames.Add(MoveTemp(UpperPropertyStr),1);
+	FName CurrentText(PropertyName,FNAME_Find); // keep generating this FName in case it has been affecting the case of future FNames.
+	return FNAME_Replace_Not_Safe_For_Threading;
+}
+
+FProperty* FHeaderParser::GetVarNameAndDim
+(
+	UStruct*                Scope,
+	FToken&                 VarProperty,
+	EVariableCategory::Type VariableCategory,
+	ELayoutMacroType        LayoutMacroType
+)
+{
+	check(Scope);
+
+	FUnrealSourceFile* CurrentSrcFile = GetCurrentSourceFile();
+	EObjectFlags ObjectFlags = RF_Public;
+	if (VariableCategory == EVariableCategory::Member && CurrentAccessSpecifier == ACCESS_Private)
+	{
+		ObjectFlags = RF_NoFlags;
+	}
+
+	const TCHAR* HintText = GetHintText(VariableCategory);
+
+	AddModuleRelativePathToMetadata(Scope, VarProperty.MetaData);
+
+	// Get variable name.
+	if (VariableCategory == EVariableCategory::Return)
+	{
+		// Hard-coded variable name, such as with return value.
+		VarProperty.TokenType = TOKEN_Identifier;
+		FCString::Strcpy( VarProperty.Identifier, TEXT("ReturnValue") );
+	}
+	else
+	{
+		FToken VarToken;
+		if (!GetIdentifier(VarToken))
+		{
+			FError::Throwf(TEXT("Missing variable name") );
+		}
+
+		switch (LayoutMacroType)
+		{
+			case ELayoutMacroType::Array:
+			case ELayoutMacroType::ArrayEditorOnly:
+			case ELayoutMacroType::Bitfield:
+			case ELayoutMacroType::BitfieldEditorOnly:
+			case ELayoutMacroType::FieldInitialized:
+				RequireSymbol(TEXT(','), GLayoutMacroNames[(int32)LayoutMacroType]);
+				break;
+
+			default:
+				break;
+		}
+
+		VarProperty.TokenType = TOKEN_Identifier;
+		FCString::Strcpy(VarProperty.Identifier, VarToken.Identifier);
+	}
+
+	// Check to see if the variable is deprecated, and if so set the flag
+	{
+		FString VarName(VarProperty.Identifier);
+
+		const int32 DeprecatedIndex = VarName.Find(TEXT("_DEPRECATED"));
+		const int32 NativizedPropertyPostfixIndex = VarName.Find(TEXT("__pf")); //TODO: check OverrideNativeName in Meta Data, to be sure it's not a random occurrence of the "__pf" string.
+		bool bIgnoreDeprecatedWord = (NativizedPropertyPostfixIndex != INDEX_NONE) && (NativizedPropertyPostfixIndex > DeprecatedIndex);
+		if ((DeprecatedIndex != INDEX_NONE) && !bIgnoreDeprecatedWord)
+		{
+			if (DeprecatedIndex != VarName.Len() - 11)
+			{
+				FError::Throwf(TEXT("Deprecated variables must end with _DEPRECATED"));
+			}
+
+			// We allow deprecated properties in blueprints that have getters and setters assigned as they may be part of a backwards compatibility path
+			const bool bBlueprintVisible = (VarProperty.PropertyFlags & CPF_BlueprintVisible) > 0;
+			const bool bWarnOnGetter = bBlueprintVisible && !VarProperty.MetaData.Contains(NAME_BlueprintGetter);
+			const bool bWarnOnSetter = bBlueprintVisible && !(VarProperty.PropertyFlags & CPF_BlueprintReadOnly) && !VarProperty.MetaData.Contains(NAME_BlueprintSetter);
+
+			if (bWarnOnGetter)
+			{
+				UE_LOG_WARNING_UHT(TEXT("%s: Deprecated property '%s' should not be marked as blueprint visible without having a BlueprintGetter"), HintText, *VarName);
+			}
+
+			if (bWarnOnSetter)
+			{
+				UE_LOG_WARNING_UHT(TEXT("%s: Deprecated property '%s' should not be marked as blueprint writeable without having a BlueprintSetter"), HintText, *VarName);
+			}
+
+
+			// Warn if a deprecated property is visible
+			if (VarProperty.PropertyFlags & (CPF_Edit | CPF_EditConst) || // Property is marked as editable
+				(!bBlueprintVisible && (VarProperty.PropertyFlags & CPF_BlueprintReadOnly) && !(VarProperty.ImpliedPropertyFlags & CPF_BlueprintReadOnly)) ) // Is BPRO, but not via Implied Flags and not caught by Getter/Setter path above
+			{
+				UE_LOG_WARNING_UHT(TEXT("%s: Deprecated property '%s' should not be marked as visible or editable"), HintText, *VarName);
+			}
+
+			VarProperty.PropertyFlags |= CPF_Deprecated;
+			VarName.MidInline(0, DeprecatedIndex, false);
+
+			FCString::Strcpy(VarProperty.Identifier, *VarName);
+		}
+	}
+
+	// Make sure it doesn't conflict.
+	int32 OuterContextCount = 0;
+	UField* ExistingField = FindField(Scope, VarProperty.Identifier, true, UField::StaticClass(), NULL);
+	FField* ExistingProperty = FindProperty(Scope, VarProperty.Identifier, true, FField::StaticClass(), NULL);
+
+	if (ExistingField != nullptr || ExistingProperty != nullptr)
+	{
+		bool bErrorDueToShadowing = true;
+
+		if (ExistingField && ExistingField->IsA(UFunction::StaticClass()) && (VariableCategory != EVariableCategory::Member))
+		{
+			// A function parameter with the same name as a method is allowed
+			bErrorDueToShadowing = false;
+		}
+
+		//@TODO: This exception does not seem sound either, but there is enough existing code that it will need to be
+		// fixed up first before the exception it is removed.
+		if (ExistingProperty)
+ 		{
+ 			FProperty* ExistingProp = CastField<FProperty>(ExistingProperty);
+ 			const bool bExistingPropDeprecated = (ExistingProp != nullptr) && ExistingProp->HasAnyPropertyFlags(CPF_Deprecated);
+ 			const bool bNewPropDeprecated = (VariableCategory == EVariableCategory::Member) && ((VarProperty.PropertyFlags & CPF_Deprecated) != 0);
+ 			if (bNewPropDeprecated || bExistingPropDeprecated)
+ 			{
+ 				// if this is a property and one of them is deprecated, ignore it since it will be removed soon
+ 				bErrorDueToShadowing = false;
+ 			}
+ 		}
+
+		if (bErrorDueToShadowing)
+		{
+			FError::Throwf(TEXT("%s: '%s' cannot be defined in '%s' as it is already defined in scope '%s' (shadowing is not allowed)"),
+				HintText,
+				VarProperty.Identifier,
+				*Scope->GetName(),
+				ExistingField ? *ExistingField->GetOuter()->GetName() : *ExistingProperty->GetOwnerVariant().GetFullName());
+		}
+	}
+
+	// Get optional dimension immediately after name.
+	FToken Dimensions;
+	if ((LayoutMacroType == ELayoutMacroType::None && MatchSymbol(TEXT('['))) || LayoutMacroType == ELayoutMacroType::Array || LayoutMacroType == ELayoutMacroType::ArrayEditorOnly)
+	{
+		switch (VariableCategory)
+		{
+			case EVariableCategory::Return:
+			{
+				FError::Throwf(TEXT("Arrays aren't allowed as return types"));
+			}
+
+			case EVariableCategory::RegularParameter:
+			case EVariableCategory::ReplicatedParameter:
+			{
+				FError::Throwf(TEXT("Arrays aren't allowed as function parameters"));
+			}
+		}
+
+		if (VarProperty.IsContainer())
+		{
+			FError::Throwf(TEXT("Static arrays of containers are not allowed"));
+		}
+
+		if (VarProperty.IsBool())
+		{
+			FError::Throwf(TEXT("Bool arrays are not allowed") );
+		}
+
+		if (LayoutMacroType == ELayoutMacroType::None)
+		{
+			// Ignore how the actual array dimensions are actually defined - we'll calculate those with the compiler anyway.
+			if (!GetRawToken(Dimensions, TEXT(']')))
+			{
+				FError::Throwf(TEXT("%s %s: Missing ']'"), HintText, VarProperty.Identifier);
+			}
+		}
+		else
+		{
+			// Ignore how the actual array dimensions are actually defined - we'll calculate those with the compiler anyway.
+			if (!GetRawToken(Dimensions, TEXT(')')))
+			{
+				FError::Throwf(TEXT("%s %s: Missing ']'"), HintText, VarProperty.Identifier);
+			}
+		}
+
+		// Only static arrays are declared with [].  Dynamic arrays use TArray<> instead.
+		VarProperty.ArrayType = EArrayType::Static;
+
+		UEnum* Enum = nullptr;
+
+		if (*Dimensions.String)
+		{
+			FString Temp = Dimensions.String;
+
+			bool bAgain;
+			do
+			{
+				bAgain = false;
+
+				// Remove any casts
+				static const TCHAR* Casts[] = {
+					TEXT("(uint32)"),
+					TEXT("(int32)"),
+					TEXT("(uint16)"),
+					TEXT("(int16)"),
+					TEXT("(uint8)"),
+					TEXT("(int8)"),
+					TEXT("(int)"),
+					TEXT("(unsigned)"),
+					TEXT("(signed)"),
+					TEXT("(unsigned int)"),
+					TEXT("(signed int)")
+				};
+
+				// Remove any brackets
+				if (Temp[0] == TEXT('('))
+				{
+					int32 TempLen      = Temp.Len();
+					int32 ClosingParen = FindMatchingClosingParenthesis(Temp);
+					if (ClosingParen == TempLen - 1)
+					{
+						Temp.MidInline(1, TempLen - 2, false);
+						bAgain = true;
+					}
+				}
+
+				for (const TCHAR* Cast : Casts)
+				{
+					if (Temp.StartsWith(Cast, ESearchCase::CaseSensitive))
+					{
+						Temp.RightChopInline(FCString::Strlen(Cast), false);
+						bAgain = true;
+					}
+				}
+			}
+			while (bAgain);
+
+			UEnum::LookupEnumNameSlow(*Temp, &Enum);
+		}
+
+		if (!Enum)
+		{
+			// If the enum wasn't declared in this scope, then try to find it anywhere we can
+			Enum = FindObject<UEnum>(ANY_PACKAGE, Dimensions.String);
+		}
+
+		if (Enum)
+		{
+			// set the ArraySizeEnum if applicable
+			VarProperty.MetaData.Add(NAME_ArraySizeEnum, Enum->GetPathName());
+		}
+
+		if (LayoutMacroType == ELayoutMacroType::None)
+		{
+			MatchSymbol(TEXT(']'));
+		}
+	}
+
+	// Try gathering metadata for member fields
+	if (VariableCategory == EVariableCategory::Member)
+	{
+		ParseFieldMetaData(VarProperty.MetaData, VarProperty.Identifier);
+		AddFormattedPrevCommentAsTooltipMetaData(VarProperty.MetaData);
+	}
+	// validate UFunction parameters
+	else
+	{
+		// UFunctions with a smart pointer as input parameter wont compile anyway, because of missing P_GET_... macro.
+		// UFunctions with a smart pointer as return type will crash when called via blueprint, because they are not supported in VM.
+		// WeakPointer is supported by VM as return type (see UObject::execLetWeakObjPtr), but there is no P_GET_... macro for WeakPointer.
+		if (VarProperty.Type == CPT_LazyObjectReference)
+		{
+			FError::Throwf(TEXT("UFunctions cannot take a lazy pointer as a parameter."));
+		}
+	}
+
+	// If this is the first time seeing the property name, then flag it for replace instead of add
+	const EFindName FindFlag = VarProperty.PropertyFlags & CPF_Config ? GetFindFlagForPropertyName(VarProperty.Identifier) : FNAME_Add;
+	// create the FName for the property, splitting (ie Unnamed_3 -> Unnamed,3)
+	FName PropertyName(VarProperty.Identifier, FindFlag);
+
+	FProperty* Prev = nullptr;
+	for (TFieldIterator<FProperty> It(Scope, EFieldIteratorFlags::ExcludeSuper); It; ++It)
+	{
+		Prev = *It;
+	}
+
+	auto PropagateFlagsFromInnerAndHandlePersistentInstanceMetadata = [](EPropertyFlags& DestFlags, const TMap<FName, FString>& InMetaData, FProperty* Inner) {
+		// Copy some of the property flags to the container property.
+		if (Inner->PropertyFlags & (CPF_ContainsInstancedReference | CPF_InstancedReference))
+		{
+			DestFlags |= CPF_ContainsInstancedReference;
+			DestFlags &= ~(CPF_InstancedReference | CPF_PersistentInstance); //this was propagated to the inner
+
+			if (Inner->PropertyFlags & CPF_PersistentInstance)
+			{
+				AddMetaDataToClassData(Inner, TMap<FName,FString>(InMetaData));
+			}
+		}
+	};
+
+	FProperty* Result = nullptr;
+	if (VarProperty.ArrayType == EArrayType::Dynamic)
+	{
+		FArrayProperty* Array     = new FArrayProperty(Scope, PropertyName, ObjectFlags);
+		FProperty*      InnerProp = CreateVariableProperty(VarProperty, Array, PropertyName, RF_Public, VariableCategory, CurrentSrcFile);
+
+		Array->Inner         = InnerProp;
+		Array->PropertyFlags = VarProperty.PropertyFlags;
+
+		// Propagate flags
+		InnerProp->PropertyFlags |= Array->PropertyFlags & CPF_PropagateToArrayInner;
+
+		PropagateFlagsFromInnerAndHandlePersistentInstanceMetadata(Array->PropertyFlags, VarProperty.MetaData, InnerProp);
+
+		Result = Array;
+
+		if (VarProperty.AllocatorType == EAllocatorType::MemoryImage)
+		{
+			GPropertyUsesMemoryImageAllocator.Add(Array);
+		}
+	}
+	else if (VarProperty.ArrayType == EArrayType::Set)
+	{
+		FSetProperty* Set       = new FSetProperty(Scope, PropertyName, ObjectFlags);
+		FProperty*    InnerProp = CreateVariableProperty(VarProperty, Set, PropertyName, RF_Public, VariableCategory, CurrentSrcFile);
+
+		Set->ElementProp   = InnerProp;
+		Set->PropertyFlags = VarProperty.PropertyFlags;
+
+		// Propagate flags
+		InnerProp->PropertyFlags |= Set->PropertyFlags & CPF_PropagateToSetElement;
+
+		PropagateFlagsFromInnerAndHandlePersistentInstanceMetadata(Set->PropertyFlags, VarProperty.MetaData, InnerProp);
+
+		Result = Set;
+	}
+	else if (VarProperty.MapKeyProp.IsValid())
+	{
+		FMapProperty* Map       = new FMapProperty(Scope, PropertyName, ObjectFlags);
+		FProperty*    KeyProp   = CreateVariableProperty(*VarProperty.MapKeyProp, Map, *(PropertyName.ToString() + TEXT("_Key")), RF_Public, VariableCategory, CurrentSrcFile);
+		FProperty*    ValueProp = CreateVariableProperty(VarProperty,             Map, PropertyName,                              RF_Public, VariableCategory, CurrentSrcFile);
+
+		Map->KeyProp       = KeyProp;
+		Map->ValueProp     = ValueProp;
+		Map->PropertyFlags = VarProperty.PropertyFlags;
+
+		// Propagate flags
+		KeyProp  ->PropertyFlags |= VarProperty.MapKeyProp->PropertyFlags & CPF_PropagateToMapKey;
+		ValueProp->PropertyFlags |= Map->PropertyFlags                    & CPF_PropagateToMapValue;
+
+		PropagateFlagsFromInnerAndHandlePersistentInstanceMetadata(Map->PropertyFlags, VarProperty.MapKeyProp->MetaData, KeyProp);
+		PropagateFlagsFromInnerAndHandlePersistentInstanceMetadata(Map->PropertyFlags, VarProperty.MetaData,             ValueProp);
+
+		Result = Map;
+
+		if (VarProperty.AllocatorType == EAllocatorType::MemoryImage)
+		{
+			GPropertyUsesMemoryImageAllocator.Add(Map);
+		}
+	}
+	else
+	{
+		Result = CreateVariableProperty(VarProperty, Scope, PropertyName, ObjectFlags, VariableCategory, CurrentSrcFile);
+
+		if (VarProperty.ArrayType == EArrayType::Static)
+		{
+			Result->ArrayDim = 2; // 2 = static array
+			GArrayDimensions.Add(Result, Dimensions.String);
+		}
+
+		Result->PropertyFlags = VarProperty.PropertyFlags;
+	}
+
+	if (Prev != nullptr)
+	{
+		Result->Next = Prev->Next;
+		Prev->Next = Result;
+	}
+	else
+	{
+		Result->Next = Scope->ChildProperties;
+		Scope->ChildProperties = Result;
+	}
+
+	VarProperty.TokenProperty = Result;
+	VarProperty.StartLine = InputLine;
+	VarProperty.StartPos = InputPos;
+	FClassMetaData* ScopeData = GScriptHelper.FindClassData(Scope);
+	check(ScopeData);
+	ScopeData->AddProperty(FToken(VarProperty), CurrentSrcFile);
+
+	// if we had any metadata, add it to the class
+	AddMetaDataToClassData(VarProperty.TokenProperty, TMap<FName,FString>(VarProperty.MetaData));
+
+	return Result;
+}
+
+/*-----------------------------------------------------------------------------
+	Statement compiler.
+-----------------------------------------------------------------------------*/
+
+//
+// Compile a declaration in Token. Returns 1 if compiled, 0 if not.
+//
+bool FHeaderParser::CompileDeclaration(FClasses& AllClasses, TArray<UDelegateFunction*>& DelegatesToFixup, FToken& Token)
+{
+	EAccessSpecifier AccessSpecifier = ParseAccessProtectionSpecifier(Token);
+	if (AccessSpecifier)
+	{
+		if (!IsAllowedInThisNesting(ENestAllowFlags::VarDecl) && !IsAllowedInThisNesting(ENestAllowFlags::Function))
+		{
+			FError::Throwf(TEXT("Access specifier %s not allowed here."), Token.Identifier);
+		}
+		check(TopNest->NestType == ENestType::Class || TopNest->NestType == ENestType::Interface || TopNest->NestType == ENestType::NativeInterface);
+		CurrentAccessSpecifier = AccessSpecifier;
+		return true;
+	}
+
+	if (Token.Matches(TEXT("class"), ESearchCase::CaseSensitive) && (TopNest->NestType == ENestType::GlobalScope))
+	{
+		// Make sure the previous class ended with valid nesting.
+		if (bEncounteredNewStyleClass_UnmatchedBrackets)
+		{
+			FError::Throwf(TEXT("Missing } at end of class"));
+		}
+
+		// Start parsing the second class
+		bEncounteredNewStyleClass_UnmatchedBrackets = true;
+		CurrentAccessSpecifier = ACCESS_Private;
+
+		if (!TryParseIInterfaceClass(AllClasses))
+		{
+			bEncounteredNewStyleClass_UnmatchedBrackets = false;
+			UngetToken(Token);
+			return SkipDeclaration(Token);
+		}
+		return true;
+	}
+
+	if (Token.Matches(TEXT("GENERATED_IINTERFACE_BODY"), ESearchCase::CaseSensitive) || (Token.Matches(TEXT("GENERATED_BODY"), ESearchCase::CaseSensitive) && TopNest->NestType == ENestType::NativeInterface))
+	{
+		if (TopNest->NestType != ENestType::NativeInterface)
+		{
+			FError::Throwf(TEXT("%s must occur inside the native interface definition"), Token.Identifier);
+		}
+		RequireSymbol(TEXT('('), Token.Identifier);
+		CompileVersionDeclaration(GetCurrentClass());
+		RequireSymbol(TEXT(')'), Token.Identifier);
+
+		FClassMetaData* ClassData = GetCurrentClassData();
+		if (!ClassData || ClassData->ParsedInterface == EParsedInterface::NotAnInterface)
+		{
+			FString CurrentClassName = GetCurrentClass()->GetName();
+			FError::Throwf(TEXT("Could not find the associated 'U%s' class while parsing 'I%s' - it could be missing or malformed"), *CurrentClassName, *CurrentClassName);
+		}
+
+		if (ClassData->ParsedInterface == EParsedInterface::ParsedIInterface)
+		{
+			FString CurrentClassName = GetCurrentClass()->GetName();
+			FError::Throwf(TEXT("Duplicate IInterface definition found while parsing 'I%s'"), *CurrentClassName);
+		}
+
+		check(ClassData->ParsedInterface == EParsedInterface::ParsedUInterface);
+
+		ClassData->ParsedInterface = EParsedInterface::ParsedIInterface;
+		ClassData->GeneratedBodyMacroAccessSpecifier = CurrentAccessSpecifier;
+		ClassData->SetInterfaceGeneratedBodyLine(InputLine);
+
+		bClassHasGeneratedIInterfaceBody = true;
+
+		if (Token.Matches(TEXT("GENERATED_IINTERFACE_BODY"), ESearchCase::CaseSensitive))
+		{
+			CurrentAccessSpecifier = ACCESS_Public;
+		}
+
+		if (Token.Matches(TEXT("GENERATED_BODY"), ESearchCase::CaseSensitive))
+		{
+			ClassDefinitionRanges[GetCurrentClass()].bHasGeneratedBody = true;
+		}
+		return true;
+	}
+
+	if (Token.Matches(TEXT("GENERATED_UINTERFACE_BODY"), ESearchCase::CaseSensitive) || (Token.Matches(TEXT("GENERATED_BODY"), ESearchCase::CaseSensitive) && TopNest->NestType == ENestType::Interface))
+	{
+		if (TopNest->NestType != ENestType::Interface)
+		{
+			FError::Throwf(TEXT("%s must occur inside the interface definition"), Token.Identifier);
+		}
+		RequireSymbol(TEXT('('), Token.Identifier);
+		CompileVersionDeclaration(GetCurrentClass());
+		RequireSymbol(TEXT(')'), Token.Identifier);
+
+		FClassMetaData* ClassData = GetCurrentClassData();
+
+		ClassData->GeneratedBodyMacroAccessSpecifier = CurrentAccessSpecifier;
+		ClassData->SetGeneratedBodyLine(InputLine);
+
+		bClassHasGeneratedUInterfaceBody = true;
+
+		if (Token.Matches(TEXT("GENERATED_UINTERFACE_BODY"), ESearchCase::CaseSensitive))
+		{
+			CurrentAccessSpecifier = ACCESS_Public;
+		}
+		return true;
+	}
+
+	if (Token.Matches(TEXT("GENERATED_UCLASS_BODY"), ESearchCase::CaseSensitive) || (Token.Matches(TEXT("GENERATED_BODY"), ESearchCase::CaseSensitive) && TopNest->NestType == ENestType::Class))
+	{
+		if (TopNest->NestType != ENestType::Class)
+		{
+			FError::Throwf(TEXT("%s must occur inside the class definition"), Token.Identifier);
+		}
+
+		FClassMetaData* ClassData = GetCurrentClassData();
+
+		if (Token.Matches(TEXT("GENERATED_BODY"), ESearchCase::CaseSensitive))
+		{
+			if (!ClassDefinitionRanges.Contains(GetCurrentClass()))
+			{
+				ClassDefinitionRanges.Add(GetCurrentClass(), ClassDefinitionRange());
+			}
+
+			ClassDefinitionRanges[GetCurrentClass()].bHasGeneratedBody = true;
+
+			ClassData->GeneratedBodyMacroAccessSpecifier = CurrentAccessSpecifier;
+		}
+		else
+		{
+			CurrentAccessSpecifier = ACCESS_Public;
+		}
+
+		RequireSymbol(TEXT('('), Token.Identifier);
+		CompileVersionDeclaration(GetCurrentClass());
+		RequireSymbol(TEXT(')'), Token.Identifier);
+
+		ClassData->SetGeneratedBodyLine(InputLine);
+
+		bClassHasGeneratedBody = true;
+		return true;
+	}
+
+	if (Token.Matches(TEXT("UCLASS"), ESearchCase::CaseSensitive))
+	{
+		bHaveSeenUClass = true;
+		bEncounteredNewStyleClass_UnmatchedBrackets = true;
+		UClass* Class = CompileClassDeclaration(AllClasses);
+		GStructToSourceLine.Add(Class, MakeTuple(GetCurrentSourceFile()->AsShared(), Token.StartLine));
+		return true;
+	}
+
+	if (Token.Matches(TEXT("UINTERFACE"), ESearchCase::CaseSensitive))
+	{
+		bHaveSeenUClass = true;
+		bEncounteredNewStyleClass_UnmatchedBrackets = true;
+		CompileInterfaceDeclaration(AllClasses);
+		return true;
+	}
+
+	if (Token.Matches(TEXT("UFUNCTION"), ESearchCase::CaseSensitive))
+	{
+		CompileFunctionDeclaration(AllClasses);
+		return true;
+	}
+
+	if (Token.Matches(TEXT("UDELEGATE"), ESearchCase::CaseSensitive))
+	{
+		UDelegateFunction* Delegate = CompileDelegateDeclaration(AllClasses, Token.Identifier, EDelegateSpecifierAction::Parse);
+		DelegatesToFixup.Add(Delegate);
+		return true;
+	}
+
+	if (IsValidDelegateDeclaration(Token)) // Legacy delegate parsing - it didn't need a UDELEGATE
+	{
+		UDelegateFunction* Delegate = CompileDelegateDeclaration(AllClasses, Token.Identifier);
+		DelegatesToFixup.Add(Delegate);
+		return true;
+	}
+
+	if (Token.Matches(TEXT("UPROPERTY"), ESearchCase::CaseSensitive))
+	{
+		CheckAllow(TEXT("'Member variable declaration'"), ENestAllowFlags::VarDecl);
+		check(TopNest->NestType == ENestType::Class);
+
+		CompileVariableDeclaration(AllClasses, GetCurrentClass());
+		return true;
+	}
+
+	if (Token.Matches(TEXT("UENUM"), ESearchCase::CaseSensitive))
+	{
+		// Enumeration definition.
+		CompileEnum();
+		return true;
+	}
+
+	if (Token.Matches(TEXT("USTRUCT"), ESearchCase::CaseSensitive))
+	{
+		// Struct definition.
+		UScriptStruct* Struct = CompileStructDeclaration(AllClasses);
+		GStructToSourceLine.Add(Struct, MakeTuple(GetCurrentSourceFile()->AsShared(), Token.StartLine));
+		return true;
+	}
+
+	if (Token.Matches(TEXT('#')))
+	{
+		// Compiler directive.
+		CompileDirective(AllClasses);
+		return true;
+	}
+
+	if (bEncounteredNewStyleClass_UnmatchedBrackets && Token.Matches(TEXT('}')))
+	{
+		if (ClassDefinitionRanges.Contains(GetCurrentClass()))
+		{
+			ClassDefinitionRanges[GetCurrentClass()].End = &Input[InputPos];
+		}
+		MatchSemi();
+
+		// Closing brace for class declaration
+		//@TODO: This is a very loose approximation of what we really need to do
+		// Instead, the whole statement-consumer loop should be in a nest
+		bEncounteredNewStyleClass_UnmatchedBrackets = false;
+
+		UClass* CurrentClass = GetCurrentClass();
+
+		// Pop nesting here to allow other non UClass declarations in the header file.
+		if (CurrentClass->ClassFlags & CLASS_Interface)
+		{
+			checkf(TopNest->NestType == ENestType::Interface || TopNest->NestType == ENestType::NativeInterface, TEXT("Unexpected end of interface block."));
+			PopNest(TopNest->NestType, TEXT("'Interface'"));
+			PostPopNestInterface(AllClasses, CurrentClass);
+
+			// Ensure the UINTERFACE classes have a GENERATED_BODY declaration
+			if (bHaveSeenUClass && !bClassHasGeneratedUInterfaceBody)
+			{
+				FError::Throwf(TEXT("Expected a GENERATED_BODY() at the start of class"));
+			}
+
+			// Ensure the non-UINTERFACE interface classes have a GENERATED_BODY declaration
+			if (!bHaveSeenUClass && !bClassHasGeneratedIInterfaceBody)
+			{
+				FError::Throwf(TEXT("Expected a GENERATED_BODY() at the start of class"));
+			}
+		}
+		else
+		{
+			PopNest(ENestType::Class, TEXT("'Class'"));
+			PostPopNestClass(CurrentClass);
+
+			// Ensure classes have a GENERATED_BODY declaration
+			if (bHaveSeenUClass && !bClassHasGeneratedBody)
+			{
+				FError::Throwf(TEXT("Expected a GENERATED_BODY() at the start of class"));
+			}
+		}
+
+		bHaveSeenUClass                  = false;
+		bClassHasGeneratedBody           = false;
+		bClassHasGeneratedUInterfaceBody = false;
+		bClassHasGeneratedIInterfaceBody = false;
+
+		GetCurrentScope()->AddType(CurrentClass);
+		return true;
+	}
+
+	if (Token.Matches(TEXT(';')))
+	{
+		if (GetToken(Token))
+		{
+			FError::Throwf(TEXT("Extra ';' before '%s'"), Token.Identifier);
+		}
+		else
+		{
+			FError::Throwf(TEXT("Extra ';' before end of file"));
+		}
+	}
+
+	if (bEncounteredNewStyleClass_UnmatchedBrackets && IsInAClass())
+	{
+		if (UClass* Class = GetCurrentClass())
+		{
+			FToken ConstructorToken = Token;
+
+			// Allow explicit constructors
+			bool bFoundExplicit = ConstructorToken.Matches(TEXT("explicit"), ESearchCase::CaseSensitive);
+			if (bFoundExplicit)
+			{
+				GetToken(ConstructorToken);
+			}
+
+			bool bSkippedAPIToken = false;
+			if (FString(ConstructorToken.Identifier).EndsWith(TEXT("_API"), ESearchCase::CaseSensitive))
+			{
+				if (!bFoundExplicit)
+				{
+					// Explicit can come before or after an _API
+					MatchIdentifier(TEXT("explicit"), ESearchCase::CaseSensitive);
+				}
+
+				GetToken(ConstructorToken);
+				bSkippedAPIToken = true;
+			}
+
+			if (ConstructorToken.Matches(*FNameLookupCPP::GetNameCPP(Class), ESearchCase::IgnoreCase))
+			{
+				if (TryToMatchConstructorParameterList(ConstructorToken))
+				{
+					return true;
+				}
+			}
+			else if (bSkippedAPIToken)
+			{
+				// We skipped over an _API token, but this wasn't a constructor so we need to unget so that subsequent code and still process it
+				UngetToken(ConstructorToken);
+			}
+		}
+	}
+
+	// Skip anything that looks like a macro followed by no bracket that we don't know about
+	if (ProbablyAnUnknownObjectLikeMacro(*this, Token))
+	{
+		return true;
+	}
+
+	// Determine if this statement is a serialize function declaration
+	if (bEncounteredNewStyleClass_UnmatchedBrackets && IsInAClass() && TopNest->NestType == ENestType::Class)
+	{
+		while (Token.Matches(TEXT("virtual"), ESearchCase::CaseSensitive) || FString(Token.Identifier).EndsWith(TEXT("_API"), ESearchCase::CaseSensitive))
+		{
+			GetToken(Token);
+		}
+
+		if (Token.Matches(TEXT("void"), ESearchCase::CaseSensitive))
+		{
+			GetToken(Token);
+			if (Token.Matches(TEXT("Serialize"), ESearchCase::CaseSensitive))
+			{
+				GetToken(Token);
+				if (Token.Matches(TEXT('(')))
+				{
+					GetToken(Token);
+
+					ESerializerArchiveType ArchiveType = ESerializerArchiveType::None;
+					if (Token.Matches(TEXT("FArchive"), ESearchCase::CaseSensitive))
+					{
+						GetToken(Token);
+						if (Token.Matches(TEXT('&')))
+						{
+							GetToken(Token);
+
+							// Allow the declaration to not define a name for the archive parameter
+							if (!Token.Matches(TEXT(')')))
+							{
+								GetToken(Token);
+							}
+
+							if (Token.Matches(TEXT(')')))
+							{
+								ArchiveType = ESerializerArchiveType::Archive;
+							}
+						}
+					}
+					else if (Token.Matches(TEXT("FStructuredArchive"), ESearchCase::CaseSensitive))
+					{
+						GetToken(Token);
+						if (Token.Matches(TEXT("::"), ESearchCase::CaseSensitive))
+						{
+							GetToken(Token);
+
+							if (Token.Matches(TEXT("FRecord"), ESearchCase::CaseSensitive))
+							{
+								GetToken(Token);
+
+								// Allow the declaration to not define a name for the slot parameter
+								if (!Token.Matches(TEXT(')')))
+								{
+									GetToken(Token);
+								}
+
+								if (Token.Matches(TEXT(')')))
+								{
+									ArchiveType = ESerializerArchiveType::StructuredArchiveRecord;
+								}
+							}
+						}
+					}
+					else if (Token.Matches(TEXT("FStructuredArchiveRecord"), ESearchCase::CaseSensitive))
+					{
+						GetToken(Token);
+
+						// Allow the declaration to not define a name for the slot parameter
+						if (!Token.Matches(TEXT(')')))
+						{
+							GetToken(Token);
+						}
+
+						if (Token.Matches(TEXT(')')))
+						{
+							ArchiveType = ESerializerArchiveType::StructuredArchiveRecord;
+						}
+					}
+
+					if (ArchiveType != ESerializerArchiveType::None)
+					{
+						// Found what we want!
+						if (CompilerDirectiveStack.Num() == 0 || (CompilerDirectiveStack.Num() == 1 && CompilerDirectiveStack[0] == ECompilerDirective::WithEditorOnlyData))
+						{
+							FString EnclosingDefine = CompilerDirectiveStack.Num() > 0 ? TEXT("WITH_EDITORONLY_DATA") : TEXT("");
+
+							UClass* CurrentClass = GetCurrentClass();
+
+							GClassSerializerMap.Add(CurrentClass, { ArchiveType, MoveTemp(EnclosingDefine) });
+						}
+						else
+						{
+							FError::Throwf(TEXT("Serialize functions must not be inside preprocessor blocks, except for WITH_EDITORONLY_DATA"));
+						}
+					}
+				}
+			}
+		}
+	}
+
+	// Ignore C++ declaration / function definition. 
+	return SkipDeclaration(Token);
+}
+
+bool FHeaderParser::SkipDeclaration(FToken& Token)
+{
+	// Store the current value of PrevComment so it can be restored after we parsed everything.
+	FString OldPrevComment(PrevComment);
+	// Consume all tokens until the end of declaration/definition has been found.
+	int32 NestedScopes = 0;
+	// Check if this is a class/struct declaration in which case it can be followed by member variable declaration.	
+	bool bPossiblyClassDeclaration = Token.Matches(TEXT("class"), ESearchCase::CaseSensitive) || Token.Matches(TEXT("struct"), ESearchCase::CaseSensitive);
+	// (known) macros can end without ; or } so use () to find the end of the declaration.
+	// However, we don't want to use it with DECLARE_FUNCTION, because we need it to be treated like a function.
+	bool bMacroDeclaration      = ProbablyAMacro(Token.Identifier) && !Token.Matches(TEXT("DECLARE_FUNCTION"), ESearchCase::CaseSensitive);
+	bool bEndOfDeclarationFound = false;
+	bool bDefinitionFound       = false;
+	TCHAR OpeningBracket = bMacroDeclaration ? TEXT('(') : TEXT('{');
+	TCHAR ClosingBracket = bMacroDeclaration ? TEXT(')') : TEXT('}');
+	bool bRetestCurrentToken = false;
+	while (bRetestCurrentToken || GetToken(Token))
+	{
+		// If we find parentheses at top-level and we think it's a class declaration then it's more likely
+		// to be something like: class UThing* GetThing();
+		if (bPossiblyClassDeclaration && NestedScopes == 0 && Token.Matches(TEXT('(')))
+		{
+			bPossiblyClassDeclaration = false;
+		}
+
+		bRetestCurrentToken = false;
+		if (Token.Matches(TEXT(';')) && NestedScopes == 0)
+		{
+			bEndOfDeclarationFound = true;
+			break;
+		}
+
+		if (!bMacroDeclaration && Token.Matches(TEXT("PURE_VIRTUAL"), ESearchCase::CaseSensitive) && NestedScopes == 0)
+		{
+			OpeningBracket = TEXT('(');
+			ClosingBracket = TEXT(')');
+		}
+
+		if (Token.Matches(OpeningBracket))
+		{
+			// This is a function definition or class declaration.
+			bDefinitionFound = true;
+			NestedScopes++;
+		}
+		else if (Token.Matches(ClosingBracket))
+		{
+			NestedScopes--;
+			if (NestedScopes == 0)
+			{
+				// Could be a class declaration in all capitals, and not a macro
+				bool bReallyEndDeclaration = true;
+				if (bMacroDeclaration)
+				{
+					FToken PossibleBracketToken;
+					GetToken(PossibleBracketToken);
+					UngetToken(Token);
+					GetToken(Token);
+
+					// If Strcmp returns 0, it is probably a class, else a macro.
+					bReallyEndDeclaration = !PossibleBracketToken.Matches(TEXT('{'));
+				}
+
+				if (bReallyEndDeclaration)
+				{
+					bEndOfDeclarationFound = true;
+					break;
+				}
+			}
+
+			if (NestedScopes < 0)
+			{
+				FError::Throwf(TEXT("Unexpected '}'. Did you miss a semi-colon?"));
+			}
+		}
+		else if (bMacroDeclaration && NestedScopes == 0)
+		{
+			bMacroDeclaration = false;
+			OpeningBracket = TEXT('{');
+			ClosingBracket = TEXT('}');
+			bRetestCurrentToken = true;
+		}
+	}
+	if (bEndOfDeclarationFound)
+	{
+		// Member variable declaration after class declaration (see bPossiblyClassDeclaration).
+		if (bPossiblyClassDeclaration && bDefinitionFound)
+		{
+			// Should syntax errors be also handled when someone declares a variable after function definition?
+			// Consume the variable name.
+			FToken VariableName;
+			if( !GetToken(VariableName, true) )
+			{
+				return false;
+			}
+			if (VariableName.TokenType != TOKEN_Identifier)
+			{
+				// Not a variable name.
+				UngetToken(VariableName);
+			}
+			else if (!SafeMatchSymbol(TEXT(';')))
+			{
+				FError::Throwf(*FString::Printf(TEXT("Unexpected '%s'. Did you miss a semi-colon?"), VariableName.Identifier));
+			}
+		}
+
+		// C++ allows any number of ';' after member declaration/definition.
+		while (SafeMatchSymbol(TEXT(';')));
+	}
+
+	PrevComment = OldPrevComment;
+	// clear the current value for comment
+	//ClearComment();
+
+	// Successfully consumed C++ declaration unless mismatched pair of brackets has been found.
+	return NestedScopes == 0 && bEndOfDeclarationFound;
+}
+
+bool FHeaderParser::SafeMatchSymbol( const TCHAR Match )
+{
+	FToken Token;
+
+	// Remember the position before the next token (this can include comments before the next symbol).
+	FScriptLocation LocationBeforeNextSymbol;
+	InitScriptLocation(LocationBeforeNextSymbol);
+
+	if (GetToken(Token, /*bNoConsts=*/ true))
+	{
+		if (Token.TokenType==TOKEN_Symbol && Token.Identifier[0] == Match && Token.Identifier[1] == 0)
+		{
+			return true;
+		}
+
+		UngetToken(Token);
+	}
+	// Return to the stored position.
+	ReturnToLocation(LocationBeforeNextSymbol);
+
+	return false;
+}
+
+FClass* FHeaderParser::ParseClassNameDeclaration(FClasses& AllClasses, FString& DeclaredClassName, FString& RequiredAPIMacroIfPresent)
+{
+	FUnrealSourceFile* CurrentSrcFile = GetCurrentSourceFile();
+	ParseNameWithPotentialAPIMacroPrefix(/*out*/ DeclaredClassName, /*out*/ RequiredAPIMacroIfPresent, TEXT("class"));
+
+	FClass* FoundClass = AllClasses.FindClass(*GetClassNameWithPrefixRemoved(*DeclaredClassName));
+	check(FoundClass);
+
+	FClassMetaData* ClassMetaData = GScriptHelper.AddClassData(FoundClass, CurrentSrcFile);
+
+	// Get parent class.
+	bool bSpecifiesParentClass = false;
+
+	// Skip optional final keyword
+	MatchIdentifier(TEXT("final"), ESearchCase::CaseSensitive);
+
+	if (MatchSymbol(TEXT(':')))
+	{
+		RequireIdentifier(TEXT("public"), ESearchCase::CaseSensitive, TEXT("class inheritance"));
+		bSpecifiesParentClass = true;
+	}
+
+	// Add class cast flag
+	FoundClass->ClassCastFlags |= ClassCastFlagMap::Get().GetCastFlag(DeclaredClassName);
+
+	if (bSpecifiesParentClass)
+	{
+		// Set the base class.
+		UClass* TempClass = GetQualifiedClass(AllClasses, TEXT("'extends'"));
+		check(TempClass);
+		// a class cannot 'extends' an interface, use 'implements'
+		if (TempClass->ClassFlags & CLASS_Interface)
+		{
+			FError::Throwf(TEXT("Class '%s' cannot extend interface '%s', use 'implements'"), *FoundClass->GetName(), *TempClass->GetName());
+		}
+
+		UClass* SuperClass = FoundClass->GetSuperClass();
+		if( SuperClass == NULL )
+		{
+			FoundClass->SetSuperStruct(TempClass);
+		}
+		else if( SuperClass != TempClass )
+		{
+			FError::Throwf(TEXT("%s's superclass must be %s, not %s"), *FoundClass->GetPathName(), *SuperClass->GetPathName(), *TempClass->GetPathName());
+		}
+
+		FoundClass->ClassCastFlags |= FoundClass->GetSuperClass()->ClassCastFlags;
+
+		// Handle additional inherited interface classes
+		while (MatchSymbol(TEXT(',')))
+		{
+			RequireIdentifier(TEXT("public"), ESearchCase::CaseSensitive, TEXT("Interface inheritance must be public"));
+
+			FString InterfaceName;
+
+			FToken Token;
+			for (;;)
+			{
+				if (!GetIdentifier(Token, true))
+				{
+					FError::Throwf(TEXT("Failed to get interface class identifier"));
+				}
+
+				InterfaceName += Token.Identifier;
+
+				// Handle templated native classes
+				if (MatchSymbol(TEXT('<')))
+				{
+					InterfaceName += TEXT('<');
+
+					int32 NestedScopes = 1;
+					while (NestedScopes)
+					{
+						if (!GetToken(Token))
+						{
+							FError::Throwf(TEXT("Unexpected end of file"));
+						}
+
+						if (Token.TokenType == TOKEN_Symbol)
+						{
+							if (Token.Matches(TEXT('<')))
+							{
+								++NestedScopes;
+							}
+							else if (Token.Matches(TEXT('>')))
+							{
+								--NestedScopes;
+							}
+						}
+
+						InterfaceName += Token.Identifier;
+					}
+				}
+
+				// Handle scoped native classes
+				if (MatchSymbol(TEXT("::")))
+				{
+					InterfaceName += TEXT("::");
+
+					// Keep reading nested identifiers
+					continue;
+				}
+
+				break;
+			}
+
+			HandleOneInheritedClass(AllClasses, FoundClass, MoveTemp(InterfaceName));
+		}
+	}
+	else if (FoundClass->GetSuperClass())
+	{
+		FError::Throwf(TEXT("class: missing 'Extends %s'"), *FoundClass->GetSuperClass()->GetName());
+	}
+
+	return FoundClass;
+}
+
+void FHeaderParser::HandleOneInheritedClass(FClasses& AllClasses, UClass* Class, FString&& InterfaceName)
+{
+	FUnrealSourceFile* CurrentSrcFile = GetCurrentSourceFile();
+	// Check for UInterface derived interface inheritance
+	if (UClass* Interface = AllClasses.FindScriptClass(InterfaceName))
+	{
+		// Try to find the interface
+		if ( !Interface->HasAnyClassFlags(CLASS_Interface) )
+		{
+			FError::Throwf(TEXT("Implements: Class %s is not an interface; Can only inherit from non-UObjects or UInterface derived interfaces"), *Interface->GetName() );
+		}
+
+		// Propagate the inheritable ClassFlags
+		Class->ClassFlags |= (Interface->ClassFlags) & CLASS_ScriptInherit;
+
+		new (Class->Interfaces) FImplementedInterface(Interface, 0, false);
+		if (Interface->HasAnyClassFlags(CLASS_Native))
+		{
+			FClassMetaData* ClassData = GScriptHelper.FindClassData(Class);
+			check(ClassData);
+			ClassData->AddInheritanceParent(Interface, CurrentSrcFile);
+		}
+	}
+	else
+	{
+		// Non-UObject inheritance
+		FClassMetaData* ClassData = GScriptHelper.FindClassData(Class);
+		check(ClassData);
+		ClassData->AddInheritanceParent(MoveTemp(InterfaceName), CurrentSrcFile);
+	}
+}
+
+/**
+ * Setups basic class settings after parsing.
+ */
+void PostParsingClassSetup(UClass* Class)
+{
+	// Cleanup after first pass.
+	FHeaderParser::ComputeFunctionParametersSize(Class);
+
+	// Set all optimization ClassFlags based on property types
+	auto HasAllOptimizationClassFlags = [Class]()
+	{
+		return (Class->HasAllClassFlags(CLASS_Config | CLASS_HasInstancedReference));
+	};
+
+	if (!HasAllOptimizationClassFlags())
+	{
+		for (TFieldIterator<FProperty> It(Class, EFieldIteratorFlags::ExcludeSuper); It; ++It)
+		{
+			if ((It->PropertyFlags & CPF_Config) != 0)
+			{
+				Class->ClassFlags |= CLASS_Config;
+				if (HasAllOptimizationClassFlags()) break;
+			}
+
+			if (It->ContainsInstancedObjectProperty())
+			{
+				Class->ClassFlags |= CLASS_HasInstancedReference;
+				if (HasAllOptimizationClassFlags()) break;
+			}
+		}
+	}
+
+	// Class needs to specify which ini file is going to be used if it contains config variables.
+	if ((Class->ClassFlags & CLASS_Config) && (Class->ClassConfigName == NAME_None))
+	{
+		// Inherit config setting from base class.
+		Class->ClassConfigName = Class->GetSuperClass() ? Class->GetSuperClass()->ClassConfigName : NAME_None;
+		if (Class->ClassConfigName == NAME_None)
+		{
+			FError::Throwf(TEXT("Classes with config / globalconfig member variables need to specify config file."));
+			Class->ClassConfigName = NAME_Engine;
+		}
+	}
+}
+
+/**
+ * Compiles a class declaration.
+ */
+UClass* FHeaderParser::CompileClassDeclaration(FClasses& AllClasses)
+{
+	// Start of a class block.
+	CheckAllow(TEXT("'class'"), ENestAllowFlags::Class);
+
+	// New-style UCLASS() syntax
+	TMap<FName, FString> MetaData;
+
+	TArray<FPropertySpecifier> SpecifiersFound;
+	ReadSpecifierSetInsideMacro(SpecifiersFound, TEXT("Class"), MetaData);
+
+	const int32 PrologFinishLine = InputLine;
+
+	// Members of classes have a default private access level in c++
+	// Setting this directly should be ok as we don't support nested classes, so the outer scope access should not need restoring
+	CurrentAccessSpecifier = ACCESS_Private;
+
+	AddFormattedPrevCommentAsTooltipMetaData(MetaData);
+
+	// New style files have the class name / extends afterwards
+	RequireIdentifier(TEXT("class"), ESearchCase::CaseSensitive, TEXT("Class declaration"));
+
+	// alignas() can come before or after the deprecation macro.
+	// We can't have both, but the compiler will catch that anyway.
+	SkipAlignasIfNecessary(*this);
+	SkipDeprecatedMacroIfNecessary(*this);
+	SkipAlignasIfNecessary(*this);
+
+	FString DeclaredClassName;
+	FString RequiredAPIMacroIfPresent;
+	
+	FClass* Class = ParseClassNameDeclaration(AllClasses, /*out*/ DeclaredClassName, /*out*/ RequiredAPIMacroIfPresent);
+	check(Class);
+	FClassDeclarationMetaData* ClassDeclarationData = &GClassDeclarations.FindChecked(Class->GetFName());
+
+	ClassDefinitionRanges.Add(Class, ClassDefinitionRange(&Input[InputPos], nullptr));
+
+	check(Class->ClassFlags == 0 || (Class->ClassFlags & ClassDeclarationData->ClassFlags) != 0);
+
+	Class->ClassFlags |= CLASS_Parsed;
+
+	PushNest(ENestType::Class, Class);
+	
+	const uint32 PrevClassFlags = Class->ClassFlags;
+	ResetClassData();
+
+	// Verify class variables haven't been filled in
+	check(Class->Children == NULL);
+	check(Class->Next == NULL);
+	check(Class->NetFields.Num() == 0);
+	check(Class->FirstOwnedClassRep == 0);
+
+	// Make sure our parent classes is parsed.
+	for (UClass* Temp = Class->GetSuperClass(); Temp; Temp = Temp->GetSuperClass())
+	{
+		bool bIsParsed = !!(Temp->ClassFlags & CLASS_Parsed);
+		bool bIsIntrinsic = !!(Temp->ClassFlags & CLASS_Intrinsic);
+		if (!(bIsParsed || bIsIntrinsic))
+		{
+			FError::Throwf(TEXT("'%s' can't be compiled: Parent class '%s' has errors"), *Class->GetName(), *Temp->GetName());
+		}
+	}
+
+	// Merge with categories inherited from the parent.
+	ClassDeclarationData->MergeClassCategories(Class);
+
+	// Class attributes.
+	FClassMetaData* ClassData = GScriptHelper.FindClassData(Class);
+	check(ClassData);
+	ClassData->SetPrologLine(PrologFinishLine);
+
+	ClassDeclarationData->MergeAndValidateClassFlags(DeclaredClassName, PrevClassFlags, Class, AllClasses);
+	Class->SetInternalFlags(EInternalObjectFlags::Native);
+
+	// Class metadata
+	MetaData.Append(ClassDeclarationData->MetaData);
+	if (ClassDeclarationData->ClassGroupNames.Num()) { MetaData.Add(NAME_ClassGroupNames, FString::Join(ClassDeclarationData->ClassGroupNames, TEXT(" "))); }
+	if (ClassDeclarationData->AutoCollapseCategories.Num()) { MetaData.Add(NAME_AutoCollapseCategories, FString::Join(ClassDeclarationData->AutoCollapseCategories, TEXT(" "))); }
+	if (ClassDeclarationData->HideCategories.Num()) { MetaData.Add(FHeaderParserNames::NAME_HideCategories, FString::Join(ClassDeclarationData->HideCategories, TEXT(" "))); }
+	if (ClassDeclarationData->ShowSubCatgories.Num()) { MetaData.Add(FHeaderParserNames::NAME_ShowCategories, FString::Join(ClassDeclarationData->ShowSubCatgories, TEXT(" "))); }
+	if (ClassDeclarationData->SparseClassDataTypes.Num()) { MetaData.Add(FHeaderParserNames::NAME_SparseClassDataTypes, FString::Join(ClassDeclarationData->SparseClassDataTypes, TEXT(" "))); }
+	if (ClassDeclarationData->HideFunctions.Num()) { MetaData.Add(NAME_HideFunctions, FString::Join(ClassDeclarationData->HideFunctions, TEXT(" "))); }
+	if (ClassDeclarationData->AutoExpandCategories.Num()) { MetaData.Add(NAME_AutoExpandCategories, FString::Join(ClassDeclarationData->AutoExpandCategories, TEXT(" "))); }
+
+	AddIncludePathToMetadata(Class, MetaData);
+	AddModuleRelativePathToMetadata(Class, MetaData);
+
+	// Register the metadata
+	AddMetaDataToClassData(Class, MoveTemp(MetaData));
+
+	// Handle the start of the rest of the class
+	RequireSymbol( TEXT('{'), TEXT("'Class'") );
+
+	// Make visible outside the package.
+	Class->ClearFlags(RF_Transient);
+	check(Class->HasAnyFlags(RF_Public));
+	check(Class->HasAnyFlags(RF_Standalone));
+
+	// Copy properties from parent class.
+	if (Class->GetSuperClass())
+	{
+		Class->SetPropertiesSize(Class->GetSuperClass()->GetPropertiesSize());
+	}
+
+	// auto-create properties for all of the VFTables needed for the multiple inheritances
+	// get the inheritance parents
+	const TArray<FMultipleInheritanceBaseClass*>& InheritanceParents = ClassData->GetInheritanceParents();
+
+	// for all base class types, make a VfTable property
+	for (int32 ParentIndex = InheritanceParents.Num() - 1; ParentIndex >= 0; ParentIndex--)
+	{
+		// if this base class corresponds to an interface class, assign the vtable FProperty in the class's Interfaces map now...
+		if (UClass* InheritedInterface = InheritanceParents[ParentIndex]->InterfaceClass)
+		{
+			FImplementedInterface* Found = Class->Interfaces.FindByPredicate([=](const FImplementedInterface& Impl) { return Impl.Class == InheritedInterface; });
+			if (Found)
+			{
+				Found->PointerOffset = 1;
+			}
+			else
+			{
+				Class->Interfaces.Add(FImplementedInterface(InheritedInterface, 1, false));
+			}
+		}
+	}
+
+	// Validate sparse class data
+	CheckSparseClassData(Class);
+
+	return Class;
+}
+
+FClass* FHeaderParser::ParseInterfaceNameDeclaration(FClasses& AllClasses, FString& DeclaredInterfaceName, FString& RequiredAPIMacroIfPresent)
+{
+	ParseNameWithPotentialAPIMacroPrefix(/*out*/ DeclaredInterfaceName, /*out*/ RequiredAPIMacroIfPresent, TEXT("interface"));
+
+	FClass* FoundClass = AllClasses.FindClass(*GetClassNameWithPrefixRemoved(*DeclaredInterfaceName));
+	if (FoundClass == nullptr)
+	{
+		return nullptr;
+	}
+
+	// Get super interface
+	bool bSpecifiesParentClass = MatchSymbol(TEXT(':'));
+	if (!bSpecifiesParentClass)
+	{
+		return FoundClass;
+	}
+
+	RequireIdentifier(TEXT("public"), ESearchCase::CaseSensitive, TEXT("class inheritance"));
+
+	// verify if our super class is an interface class
+	// the super class should have been marked as CLASS_Interface at the importing stage, if it were an interface
+	UClass* TempClass = GetQualifiedClass(AllClasses, TEXT("'extends'"));
+	check(TempClass);
+	if( !(TempClass->ClassFlags & CLASS_Interface) )
+	{
+		// UInterface is special and actually extends from UObject, which isn't an interface
+		if (DeclaredInterfaceName != TEXT("UInterface"))
+			FError::Throwf(TEXT("Interface class '%s' cannot inherit from non-interface class '%s'"), *DeclaredInterfaceName, *TempClass->GetName() );
+	}
+
+	UClass* SuperClass = FoundClass->GetSuperClass();
+	if (SuperClass == NULL)
+	{
+		FoundClass->SetSuperStruct(TempClass);
+	}
+	else if (SuperClass != TempClass)
+	{
+		FError::Throwf(TEXT("%s's superclass must be %s, not %s"), *FoundClass->GetPathName(), *SuperClass->GetPathName(), *TempClass->GetPathName());
+	}
+
+	return FoundClass;
+}
+
+bool FHeaderParser::TryParseIInterfaceClass(FClasses& AllClasses)
+{
+	// 'class' was already matched by the caller
+
+	// Get a class name
+	FString DeclaredInterfaceName;
+	FString RequiredAPIMacroIfPresent;
+	if (ParseInterfaceNameDeclaration(AllClasses, /*out*/ DeclaredInterfaceName, /*out*/ RequiredAPIMacroIfPresent) == nullptr)
+	{
+		return false;
+	}
+
+	if (MatchSymbol(TEXT(';')))
+	{
+		// Forward declaration.
+		return false;
+	}
+
+	if (DeclaredInterfaceName[0] != 'I')
+	{
+		return false;
+	}
+
+	UClass* FoundClass = nullptr;
+	if ((FoundClass = AllClasses.FindClass(*DeclaredInterfaceName.Mid(1))) == nullptr)
+	{
+		return false;
+	}
+
+	// Continue parsing the second class as if it were a part of the first (for reflection data purposes, it is)
+	RequireSymbol(TEXT('{'), TEXT("C++ interface mix-in class declaration"));
+
+	// Push the interface class nesting again.
+	PushNest(ENestType::NativeInterface, FoundClass);
+
+	return true;
+}
+
+/**
+ *  compiles Java or C# style interface declaration
+ */
+void FHeaderParser::CompileInterfaceDeclaration(FClasses& AllClasses)
+{
+	FUnrealSourceFile* CurrentSrcFile = GetCurrentSourceFile();
+	// Start of an interface block. Since Interfaces and Classes are always at the same nesting level,
+	// whereever a class declaration is allowed, an interface declaration is also allowed.
+	CheckAllow( TEXT("'interface'"), ENestAllowFlags::Class );
+
+	FString DeclaredInterfaceName;
+	FString RequiredAPIMacroIfPresent;
+	TMap<FName, FString> MetaData;
+
+	// Build up a list of interface specifiers
+	TArray<FPropertySpecifier> SpecifiersFound;
+
+	// New-style UINTERFACE() syntax
+	ReadSpecifierSetInsideMacro(SpecifiersFound, TEXT("Interface"), MetaData);
+
+	int32 PrologFinishLine = InputLine;
+
+	// New style files have the interface name / extends afterwards
+	RequireIdentifier(TEXT("class"), ESearchCase::CaseSensitive, TEXT("Interface declaration"));
+	FClass* InterfaceClass = ParseInterfaceNameDeclaration(AllClasses, /*out*/ DeclaredInterfaceName, /*out*/ RequiredAPIMacroIfPresent);
+	ClassDefinitionRanges.Add(InterfaceClass, ClassDefinitionRange(&Input[InputPos], nullptr));
+
+	// Record that this interface is RequiredAPI if the CORE_API style macro was present
+	if (!RequiredAPIMacroIfPresent.IsEmpty())
+	{
+		InterfaceClass->ClassFlags |= CLASS_RequiredAPI;
+	}
+
+	// Set the appropriate interface class flags
+	InterfaceClass->ClassFlags |= CLASS_Interface | CLASS_Abstract;
+	if (InterfaceClass->GetSuperClass() != NULL)
+	{
+		InterfaceClass->ClassCastFlags |= InterfaceClass->GetSuperClass()->ClassCastFlags;
+	}
+
+	// All classes that are parsed are expected to be native
+	if (InterfaceClass->GetSuperClass() && !InterfaceClass->GetSuperClass()->HasAnyClassFlags(CLASS_Native))
+	{
+		FError::Throwf(TEXT("Native classes cannot extend non-native classes") );
+	}
+
+	InterfaceClass->SetInternalFlags(EInternalObjectFlags::Native);
+	InterfaceClass->ClassFlags |= CLASS_Native;
+
+	// Process all of the interface specifiers
+	for (const FPropertySpecifier& Specifier : SpecifiersFound)
+	{
+		switch ((EInterfaceSpecifier)Algo::FindSortedStringCaseInsensitive(*Specifier.Key, GInterfaceSpecifierStrings))
+		{
+			default:
+			{
+				FError::Throwf(TEXT("Unknown interface specifier '%s'"), *Specifier.Key);
+			}
+			break;
+
+			case EInterfaceSpecifier::DependsOn:
+			{
+				FError::Throwf(TEXT("The dependsOn specifier is deprecated. Please use #include \"ClassHeaderFilename.h\" instead."));
+			}
+			break;
+
+			case EInterfaceSpecifier::MinimalAPI:
+			{
+				InterfaceClass->ClassFlags |= CLASS_MinimalAPI;
+			}
+			break;
+
+			case EInterfaceSpecifier::ConversionRoot:
+			{
+				MetaData.Add(FHeaderParserNames::NAME_IsConversionRoot, TEXT("true"));
+			}
+			break;
+		}
+	}
+
+	// All classes must start with a valid Unreal prefix
+	const FString ExpectedInterfaceName = InterfaceClass->GetNameWithPrefix(EEnforceInterfacePrefix::U);
+	if (DeclaredInterfaceName != ExpectedInterfaceName)
+	{
+		FError::Throwf(TEXT("Interface name '%s' is invalid, the first class should be identified as '%s'"), *DeclaredInterfaceName, *ExpectedInterfaceName );
+	}
+
+	// Try parsing metadata for the interface
+	FClassMetaData* ClassData = GScriptHelper.AddInterfaceClassData(InterfaceClass, CurrentSrcFile);
+	check(ClassData);
+
+	ClassData->SetPrologLine(PrologFinishLine);
+
+	// Register the metadata
+	AddModuleRelativePathToMetadata(InterfaceClass, MetaData);
+	AddMetaDataToClassData(InterfaceClass, MoveTemp(MetaData));
+
+	// Handle the start of the rest of the interface
+	RequireSymbol( TEXT('{'), TEXT("'Class'") );
+
+	// Make visible outside the package.
+	InterfaceClass->ClearFlags(RF_Transient);
+	check(InterfaceClass->HasAnyFlags(RF_Public));
+	check(InterfaceClass->HasAnyFlags(RF_Standalone));
+
+	// Push the interface class nesting.
+	// we need a more specific set of allow flags for ENestType::Interface, only function declaration is allowed, no other stuff are allowed
+	PushNest(ENestType::Interface, InterfaceClass);
+}
+
+void FHeaderParser::CompileRigVMMethodDeclaration(FClasses& AllClasses, UStruct* Struct)
+{
+	if (!MatchSymbol(TEXT("(")))
+	{
+		FError::Throwf(TEXT("Bad RIGVM_METHOD definition"));
+	}
+
+	// find the next close brace
+	while (!MatchSymbol(TEXT(")")))
+	{
+		FToken Token;
+		if (!GetToken(Token))
+		{
+			break;
+		}
+	}
+
+	FToken PrefixToken, ReturnTypeToken, NameToken, PostfixToken;
+	if (!GetToken(PrefixToken))
+	{
+		return;
+	}
+
+	if (FString(PrefixToken.Identifier).Equals(TEXT("virtual")))
+	{
+		if (!GetToken(ReturnTypeToken))
+		{
+			return;
+		}
+	}
+	else
+	{
+		ReturnTypeToken = PrefixToken;
+	}
+
+	if (!GetToken(NameToken))
+	{
+		return;
+	}
+
+	if (!MatchSymbol(TEXT("(")))
+	{
+		FError::Throwf(TEXT("Bad RIGVM_METHOD definition"));
+	}
+
+	TArray<FString> ParamsContent;
+	while (!MatchSymbol(TEXT(")")))
+	{
+		FToken Token;
+		if (!GetToken(Token))
+		{
+			break;
+		}
+		ParamsContent.Add(FString(Token.Identifier));
+	}
+
+	while (!FString(PostfixToken.Identifier).Equals(TEXT(";")))
+	{
+		if (!GetToken(PostfixToken))
+		{
+			return;
+		}
+	}
+
+	FRigVMMethodInfo MethodInfo;
+	MethodInfo.ReturnType = ReturnTypeToken.Identifier;
+	MethodInfo.Name = NameToken.Identifier;
+	
+	FString ParamString = FString::Join(ParamsContent, TEXT(" "));
+	if (!ParamString.IsEmpty())
+	{
+		FString ParamPrev, ParamLeft, ParamRight;
+		ParamPrev = ParamString;
+		while (ParamPrev.Contains(TEXT(",")))
+		{
+			ParamPrev.Split(TEXT(","), &ParamLeft, &ParamRight);
+			FRigVMParameter Parameter;
+			Parameter.Name = ParamLeft.TrimStartAndEnd();
+			MethodInfo.Parameters.Add(Parameter);
+			ParamPrev = ParamRight;
+		}
+
+		ParamPrev = ParamPrev.TrimStartAndEnd();
+		if (!ParamPrev.IsEmpty())
+		{
+			FRigVMParameter Parameter;
+			Parameter.Name = ParamPrev.TrimStartAndEnd();
+			MethodInfo.Parameters.Add(Parameter);
+		}
+	}
+
+	for (FRigVMParameter& Parameter : MethodInfo.Parameters)
+	{
+		FString FullParameter = Parameter.Name;
+
+		int32 LastEqual = INDEX_NONE;
+		if (FullParameter.FindLastChar(TCHAR('='), LastEqual))
+		{
+			FullParameter = FullParameter.Mid(0, LastEqual);
+		}
+
+		FullParameter.TrimStartAndEndInline();
+
+		FString ParameterType = FullParameter;
+		FString ParameterName = FullParameter;
+
+		int32 LastSpace = INDEX_NONE;
+		if (FullParameter.FindLastChar(TCHAR(' '), LastSpace))
+		{
+			Parameter.Type = FullParameter.Mid(0, LastSpace);
+			Parameter.Name = FullParameter.Mid(LastSpace + 1);
+			Parameter.Type.TrimStartAndEndInline();
+			Parameter.Name.TrimStartAndEndInline();
+		}
+	}
+
+	FRigVMStructInfo& StructRigVMInfo = StructRigVMMap.FindOrAdd(Struct);
+	StructRigVMInfo.Name = Struct->GetName();
+	StructRigVMInfo.Methods.Add(MethodInfo);
+}
+
+static const FName NAME_InputText(TEXT("Input"));
+static const FName NAME_OutputText(TEXT("Output"));
+static const FName NAME_ConstantText(TEXT("Constant"));
+static const FName NAME_MaxArraySizeText(TEXT("MaxArraySize"));
+
+static const TCHAR* TArrayText = TEXT("TArray");
+static const TCHAR* TArrayViewText = TEXT("TArrayView");
+static const TCHAR* GetRefText = TEXT("GetRef");
+static const TCHAR* GetArrayText = TEXT("GetArray");
+
+void FHeaderParser::ParseRigVMMethodParameters(UStruct* Struct)
+{
+	FRigVMStructInfo* StructRigVMInfo = StructRigVMMap.Find(Struct);
+	if (StructRigVMInfo == nullptr)
+	{
+		return;
+	}
+
+	// validate the property types for this struct
+	for (TFieldIterator<FProperty> It(Struct); It; ++It)
+	{
+		FProperty const* const Prop = *It;
+		FString MemberCPPType;
+		FString ExtendedCPPType;
+		MemberCPPType = Prop->GetCPPType(&ExtendedCPPType);
+
+		FRigVMParameter Parameter;
+		Parameter.Name = Prop->GetName();
+		Parameter.Type = MemberCPPType + ExtendedCPPType;
+		Parameter.bConstant = Prop->HasMetaData(NAME_ConstantText);
+		Parameter.bInput = Prop->HasMetaData(NAME_InputText);
+		Parameter.bOutput = Prop->HasMetaData(NAME_OutputText);
+		Parameter.MaxArraySize = Prop->GetMetaData(NAME_MaxArraySizeText);
+		Parameter.Getter = GetRefText;
+		Parameter.bEditorOnly = Prop->IsEditorOnlyProperty();
+
+		if (Parameter.bEditorOnly)
+		{
+			UE_LOG_ERROR_UHT(TEXT("RigVM Struct '%s' - Member '%s' is editor only - WITH_EDITORONLY_DATA not allowed on structs with RIGVM_METHOD."), *Struct->GetName(), *Parameter.Name, *MemberCPPType);
+		}
+
+		if (!ExtendedCPPType.IsEmpty())
+		{
+			// we only support arrays - no maps or similar data structures
+			if (MemberCPPType != TArrayText)
+			{
+				UE_LOG_ERROR_UHT(TEXT("RigVM Struct '%s' - Member '%s' type '%s' not supported by RigVM."), *Struct->GetName(), *Parameter.Name, *MemberCPPType);
+				continue;
+			}
+
+			if (!Parameter.IsConst() && Parameter.MaxArraySize.IsEmpty())
+			{
+				UE_LOG_ERROR_UHT(TEXT("RigVM Struct '%s' - Member '%s' requires the 'MaxArraySize' meta tag."), *Struct->GetName(), *Parameter.Name);
+				continue;
+			}
+		}
+
+		if (MemberCPPType.StartsWith(TArrayText, ESearchCase::CaseSensitive))
+		{
+			if (Parameter.IsConst() || !Parameter.MaxArraySize.IsEmpty())
+			{
+				Parameter.CastName = FString::Printf(TEXT("%s_%d_View"), *Parameter.Name, StructRigVMInfo->Members.Num());
+				Parameter.CastType = FString::Printf(TEXT("%s%s"), TArrayViewText, *ExtendedCPPType);
+				Parameter.Getter = GetArrayText;
+			}
+		}
+
+		StructRigVMInfo->Members.Add(MoveTemp(Parameter));
+	}
+
+	if (StructRigVMInfo->Members.Num() == 0)
+	{
+		UE_LOG_ERROR_UHT(TEXT("RigVM Struct '%s' - has zero members - invalid RIGVM_METHOD."), *Struct->GetName());
+	}
+
+	if (StructRigVMInfo->Members.Num() > 64)
+	{
+		UE_LOG_ERROR_UHT(TEXT("RigVM Struct '%s' - has %d members (64 is the limit)."), *Struct->GetName(), StructRigVMInfo->Members.Num());
+	}
+}
+
+// Returns true if the token is a dynamic delegate declaration
+bool FHeaderParser::IsValidDelegateDeclaration(const FToken& Token) const
+{
+	return (Token.TokenType == TOKEN_Identifier) && !FCString::Strncmp(Token.Identifier, TEXT("DECLARE_DYNAMIC_"), 16);
+}
+
+// Modify token to fix redirected types if needed
+void FHeaderParser::RedirectTypeIdentifier(FToken& Token) const
+{
+	check(Token.TokenType == TOKEN_Identifier);
+
+	FString* FoundRedirect = TypeRedirectMap.Find(Token.Identifier);
+	if (FoundRedirect)
+	{
+		Token.SetIdentifier(**FoundRedirect);
+	}
+}
+
+// Parse the parameter list of a function or delegate declaration
+void FHeaderParser::ParseParameterList(FClasses& AllClasses, UFunction* Function, bool bExpectCommaBeforeName, TMap<FName, FString>* MetaData)
+{
+	// Get parameter list.
+	if (MatchSymbol(TEXT(')')))
+	{
+		return;
+	}
+
+	FAdvancedDisplayParameterHandler AdvancedDisplay(MetaData);
+	do
+	{
+		// Get parameter type.
+		FToken Property(CPT_None);
+		EVariableCategory::Type VariableCategory = (Function->FunctionFlags & FUNC_Net) ? EVariableCategory::ReplicatedParameter : EVariableCategory::RegularParameter;
+		GetVarType(AllClasses, GetCurrentScope(), Property, ~(CPF_ParmFlags | CPF_AutoWeak | CPF_RepSkip | CPF_UObjectWrapper | CPF_NativeAccessSpecifiers), NULL, EPropertyDeclarationStyle::None, VariableCategory);
+		Property.PropertyFlags |= CPF_Parm;
+
+		if (bExpectCommaBeforeName)
+		{
+			RequireSymbol(TEXT(','), TEXT("Delegate definitions require a , between the parameter type and parameter name"));
+		}
+
+		FProperty* Prop = GetVarNameAndDim(Function, Property, VariableCategory);
+
+		Function->NumParms++;
+
+		if( AdvancedDisplay.CanMarkMore() && AdvancedDisplay.ShouldMarkParameter(Prop->GetName()) )
+		{
+			Prop->PropertyFlags |= CPF_AdvancedDisplay;
+		}
+
+		// Check parameters.
+		if ((Function->FunctionFlags & FUNC_Net))
+		{
+			if (!(Function->FunctionFlags & FUNC_NetRequest))
+			{
+				if (Property.PropertyFlags & CPF_OutParm)
+				{
+					UE_LOG_ERROR_UHT(TEXT("Replicated functions cannot contain out parameters"));
+				}
+
+				if (Property.PropertyFlags & CPF_RepSkip)
+				{
+					UE_LOG_ERROR_UHT(TEXT("Only service request functions cannot contain NoReplication parameters"));
+				}
+
+				if ((Prop->GetCastFlags() & CASTCLASS_FDelegateProperty) != 0)
+				{
+					UE_LOG_ERROR_UHT(TEXT("Replicated functions cannot contain delegate parameters (this would be insecure)"));
+				}
+
+				if (Property.Type == CPT_String && Property.RefQualifier != ERefQualifier::ConstRef && Prop->ArrayDim == 1)
+				{
+					UE_LOG_ERROR_UHT(TEXT("Replicated FString parameters must be passed by const reference"));
+				}
+
+				if (Property.ArrayType == EArrayType::Dynamic && Property.RefQualifier != ERefQualifier::ConstRef && Prop->ArrayDim == 1)
+				{
+					UE_LOG_ERROR_UHT(TEXT("Replicated TArray parameters must be passed by const reference"));
+				}
+			}
+			else
+			{
+				if (!(Property.PropertyFlags & CPF_RepSkip) && (Property.PropertyFlags & CPF_OutParm))
+				{
+					UE_LOG_ERROR_UHT(TEXT("Service request functions cannot contain out parameters, unless marked NotReplicated"));
+				}
+
+				if (!(Property.PropertyFlags & CPF_RepSkip) && (Prop->GetCastFlags() & CASTCLASS_FDelegateProperty) != 0)
+				{
+					UE_LOG_ERROR_UHT(TEXT("Service request functions cannot contain delegate parameters, unless marked NotReplicated"));
+				}
+			}
+		}
+		if ((Function->FunctionFlags & (FUNC_BlueprintEvent|FUNC_BlueprintCallable)) != 0)
+		{
+			if (Property.Type == CPT_Byte)
+			{
+				if (FEnumProperty* EnumProperty = CastField<FEnumProperty>(Prop))
+				{
+					FProperty* InnerType = EnumProperty->GetUnderlyingProperty();
+					if (InnerType && !InnerType->IsA<FByteProperty>())
+					{
+						FError::Throwf(TEXT("Invalid enum param for Blueprints - currently only uint8 supported"));
+					}
+				}
+			}
+
+			// Check that the parameter name is valid and does not conflict with pre-defined types
+			{
+				const static TArray<FString> InvalidParamNames =
+				{
+					TEXT("self"),
+				};
+
+				for (const FString& InvalidName : InvalidParamNames)
+				{
+					if (Property.Matches(*InvalidName, ESearchCase::IgnoreCase))
+					{
+						UE_LOG_ERROR_UHT(TEXT("Paramater name '%s' in function is invalid, '%s' is a reserved name."), *InvalidName, *InvalidName);
+					}
+				}
+			}
+		}
+
+		// Default value.
+		if (MatchSymbol(TEXT('=')))
+		{
+			// Skip past the native specified default value; we make no attempt to parse it
+			FToken SkipToken;
+			int32 ParenthesisNestCount=0;
+			int32 StartPos=-1;
+			int32 EndPos=-1;
+			while ( GetToken(SkipToken) )
+			{
+				if (StartPos == -1)
+				{
+					StartPos = SkipToken.StartPos;
+				}
+				if ( ParenthesisNestCount == 0
+					&& (SkipToken.Matches(TEXT(')')) || SkipToken.Matches(TEXT(','))) )
+				{
+					EndPos = SkipToken.StartPos;
+					// went too far
+					UngetToken(SkipToken);
+					break;
+				}
+				if ( SkipToken.Matches(TEXT('(')) )
+				{
+					ParenthesisNestCount++;
+				}
+				else if ( SkipToken.Matches(TEXT(')')) )
+				{
+					ParenthesisNestCount--;
+				}
+			}
+
+			// allow exec functions to be added to the metaData, this is so we can have default params for them.
+			const bool bStoreCppDefaultValueInMetaData = Function->HasAnyFunctionFlags(FUNC_BlueprintCallable | FUNC_Exec);
+				
+			if((EndPos > -1) && bStoreCppDefaultValueInMetaData) 
+			{
+				FString DefaultArgText(EndPos - StartPos, Input + StartPos);
+				FString Key(TEXT("CPP_Default_"));
+				Key += Prop->GetName();
+				FName KeyName = FName(*Key);
+				if (!MetaData->Contains(KeyName))
+				{
+					FString InnerDefaultValue;
+					const bool bDefaultValueParsed = DefaultValueStringCppFormatToInnerFormat(Prop, DefaultArgText, InnerDefaultValue);
+					if (!bDefaultValueParsed)
+					{
+						FError::Throwf(TEXT("C++ Default parameter not parsed: %s \"%s\" "), *Prop->GetName(), *DefaultArgText);
+					}
+
+					MetaData->Add(KeyName, InnerDefaultValue);
+						UE_LOG(LogCompile, Verbose, TEXT("C++ Default parameter parsed: %s \"%s\" -> \"%s\" "), *Prop->GetName(), *DefaultArgText, *InnerDefaultValue );
+				}
+			}
+		}
+	} while( MatchSymbol(TEXT(',')) );
+	RequireSymbol( TEXT(')'), TEXT("parameter list") );
+}
+UDelegateFunction* FHeaderParser::CompileDelegateDeclaration(FClasses& AllClasses, const TCHAR* DelegateIdentifier, EDelegateSpecifierAction::Type SpecifierAction)
+{
+	const TCHAR* CurrentScopeName = TEXT("Delegate Declaration");
+
+	FUnrealSourceFile* CurrentSrcFile = GetCurrentSourceFile();
+	TMap<FName, FString> MetaData;
+	AddModuleRelativePathToMetadata(*CurrentSrcFile, MetaData);
+
+	FFuncInfo            FuncInfo;
+
+	// If this is a UDELEGATE, parse the specifiers first
+	FString DelegateMacro;
+	if (SpecifierAction == EDelegateSpecifierAction::Parse)
+	{
+		TArray<FPropertySpecifier> SpecifiersFound;
+		ReadSpecifierSetInsideMacro(SpecifiersFound, TEXT("Delegate"), MetaData);
+
+		ProcessFunctionSpecifiers(FuncInfo, SpecifiersFound, MetaData);
+
+		// Get the next token and ensure it looks like a delegate
+		FToken Token;
+		GetToken(Token);
+		if (!IsValidDelegateDeclaration(Token))
+		{
+			FError::Throwf(TEXT("Unexpected token following UDELEGATE(): %s"), Token.Identifier);
+		}
+
+		DelegateMacro = Token.Identifier;
+
+		//Workaround for UE-28897
+		const FStructScope* CurrentStructScope = TopNest->GetScope() ? TopNest->GetScope()->AsStructScope() : nullptr;
+		const bool bDynamicClassScope = CurrentStructScope && CurrentStructScope->GetStruct() && FClass::IsDynamic(CurrentStructScope->GetStruct());
+		CheckAllow(CurrentScopeName, bDynamicClassScope ? ENestAllowFlags::ImplicitDelegateDecl : ENestAllowFlags::TypeDecl);
+	}
+	else
+	{
+		DelegateMacro = DelegateIdentifier;
+		CheckAllow(CurrentScopeName, ENestAllowFlags::ImplicitDelegateDecl);
+	}
+
+	// Break the delegate declaration macro down into parts
+	const bool bHasReturnValue = DelegateMacro.Contains(TEXT("_RetVal"), ESearchCase::CaseSensitive);
+	const bool bDeclaredConst  = DelegateMacro.Contains(TEXT("_Const"), ESearchCase::CaseSensitive);
+	const bool bIsMulticast    = DelegateMacro.Contains(TEXT("_MULTICAST"), ESearchCase::CaseSensitive);
+	const bool bIsSparse       = DelegateMacro.Contains(TEXT("_SPARSE"), ESearchCase::CaseSensitive);
+
+	// Determine the parameter count
+	const FString* FoundParamCount = DelegateParameterCountStrings.FindByPredicate([&](const FString& Str){ return DelegateMacro.Contains(Str); });
+
+	// Try reconstructing the string to make sure it matches our expectations
+	FString ExpectedOriginalString = FString::Printf(TEXT("DECLARE_DYNAMIC%s%s_DELEGATE%s%s%s"),
+		bIsMulticast ? TEXT("_MULTICAST") : TEXT(""),
+		bIsSparse ? TEXT("_SPARSE") : TEXT(""),
+		bHasReturnValue ? TEXT("_RetVal") : TEXT(""),
+		FoundParamCount ? **FoundParamCount : TEXT(""),
+		bDeclaredConst ? TEXT("_Const") : TEXT(""));
+
+	if (DelegateMacro != ExpectedOriginalString)
+	{
+		FError::Throwf(TEXT("Unable to parse delegate declaration; expected '%s' but found '%s'."), *ExpectedOriginalString, *DelegateMacro);
+	}
+
+	// Multi-cast delegate function signatures are not allowed to have a return value
+	if (bHasReturnValue && bIsMulticast)
+	{
+		UE_LOG_ERROR_UHT(TEXT("Multi-cast delegates function signatures must not return a value"));
+	}
+
+	// Delegate signature
+	FuncInfo.FunctionFlags |= FUNC_Public | FUNC_Delegate;
+
+	if (bIsMulticast)
+	{
+		FuncInfo.FunctionFlags |= FUNC_MulticastDelegate;
+	}
+
+	// Now parse the macro body
+	RequireSymbol(TEXT('('), CurrentScopeName);
+
+	// Parse the return value type
+	FToken ReturnType( CPT_None );
+
+	if (bHasReturnValue)
+	{
+		GetVarType(AllClasses, GetCurrentScope(), ReturnType, CPF_None, nullptr, EPropertyDeclarationStyle::None, EVariableCategory::Return);
+		RequireSymbol(TEXT(','), CurrentScopeName);
+	}
+
+	// Skip whitespaces to get InputPos exactly on beginning of function name.
+	while (FChar::IsWhitespace(PeekChar())) { GetChar(); }
+
+	FuncInfo.InputPos = InputPos;
+
+	// Get the delegate name
+	if (!GetIdentifier(FuncInfo.Function))
+	{
+		FError::Throwf(TEXT("Missing name for %s"), CurrentScopeName );
+	}
+
+	// If this is a delegate function then go ahead and mangle the name so we don't collide with
+	// actual functions or properties
+	{
+		//@TODO: UCREMOVAL: Eventually this mangling shouldn't occur
+
+		// Remove the leading F
+		FString Name(FuncInfo.Function.Identifier);
+
+		if (!Name.StartsWith(TEXT("F"), ESearchCase::CaseSensitive))
+		{
+			FError::Throwf(TEXT("Delegate type declarations must start with F"));
+		}
+
+		Name.RightChopInline(1, false);
+
+		// Append the signature goo
+		Name += HEADER_GENERATED_DELEGATE_SIGNATURE_SUFFIX;
+
+		// Replace the name
+		FCString::Strcpy( FuncInfo.Function.Identifier, *Name );
+	}
+
+	UDelegateFunction* DelegateSignatureFunction = (bIsSparse ? CreateDelegateFunction<USparseDelegateFunction>(FuncInfo) : CreateDelegateFunction<UDelegateFunction>(FuncInfo));
+	
+	FClassMetaData* ClassMetaData = GScriptHelper.AddClassData(DelegateSignatureFunction, CurrentSrcFile);
+
+	DelegateSignatureFunction->FunctionFlags |= FuncInfo.FunctionFlags;
+
+	FuncInfo.FunctionReference = DelegateSignatureFunction;
+	FuncInfo.SetFunctionNames();
+	if (FuncInfo.FunctionReference->HasAnyFunctionFlags(FUNC_Delegate) && !GetCurrentScope()->IsFileScope())
+	{
+		GetCurrentClassData()->MarkContainsDelegate();
+	}
+
+	GetCurrentScope()->AddType(DelegateSignatureFunction);
+
+	// determine whether this function should be 'const'
+	if (bDeclaredConst)
+	{
+		DelegateSignatureFunction->FunctionFlags |= FUNC_Const;
+	}
+
+	if (bIsSparse)
+	{
+		FToken OwningClass;
+
+		RequireSymbol(TEXT(','), TEXT("Delegate Declaration"));
+
+		if (!GetIdentifier(OwningClass))
+		{
+			FError::Throwf(TEXT("Missing OwningClass specifier."));
+		}
+		RequireSymbol(TEXT(','), TEXT("Delegate Declaration"));
+		
+		FToken DelegateName;
+		if (!GetIdentifier(DelegateName))
+		{
+			FError::Throwf(TEXT("Missing Delegate Name."));
+		}
+
+		USparseDelegateFunction* SDF = CastChecked<USparseDelegateFunction>(DelegateSignatureFunction);
+		SDF->OwningClassName = *GetClassNameWithoutPrefix(FString(OwningClass.Identifier));
+		SDF->DelegateName = DelegateName.Identifier;
+	}
+
+	// Get parameter list.
+	if (FoundParamCount)
+	{
+		RequireSymbol(TEXT(','), CurrentScopeName);
+
+		ParseParameterList(AllClasses, DelegateSignatureFunction, /*bExpectCommaBeforeName=*/ true);
+
+		// Check the expected versus actual number of parameters
+		int32 ParamCount = UE_PTRDIFF_TO_INT32(FoundParamCount - DelegateParameterCountStrings.GetData()) + 1;
+		if (DelegateSignatureFunction->NumParms != ParamCount)
+		{
+			FError::Throwf(TEXT("Expected %d parameters but found %d parameters"), ParamCount, DelegateSignatureFunction->NumParms);
+		}
+	}
+	else
+	{
+		// Require the closing paren even with no parameter list
+		RequireSymbol(TEXT(')'), TEXT("Delegate Declaration"));
+	}
+
+	// Save off function identifier in case we need it to error later, since we are moving the FuncInfo
+	TCHAR FuncInfoFunctionIdentifier[NAME_SIZE];
+	FCString::Strncpy(FuncInfoFunctionIdentifier, FuncInfo.Function.Identifier, NAME_SIZE);
+
+	FuncInfo.MacroLine = InputLine;
+	FFunctionData::Add(MoveTemp(FuncInfo));
+
+	// Create the return value property
+	if (bHasReturnValue)
+	{
+		ReturnType.PropertyFlags |= CPF_Parm | CPF_OutParm | CPF_ReturnParm;
+		FProperty* ReturnProp = GetVarNameAndDim(DelegateSignatureFunction, ReturnType, EVariableCategory::Return);
+
+		DelegateSignatureFunction->NumParms++;
+	}
+
+	// Try parsing metadata for the function
+	ParseFieldMetaData(MetaData, *(DelegateSignatureFunction->GetName()));
+
+	AddFormattedPrevCommentAsTooltipMetaData(MetaData);
+
+	AddMetaDataToClassData(DelegateSignatureFunction, MoveTemp(MetaData));
+
+	// Optionally consume a semicolon, it's not required for the delegate macro since it contains one internally
+	MatchSemi();
+
+	// Bind the function.
+	DelegateSignatureFunction->Bind();
+
+	// End the nesting
+	PostPopFunctionDeclaration(AllClasses, DelegateSignatureFunction);
+
+	// Don't allow delegate signatures to be redefined.
+	auto FunctionIterator = GetCurrentScope()->GetTypeIterator<UFunction>();
+	while (FunctionIterator.MoveNext())
+	{
+		UFunction* TestFunc = *FunctionIterator;
+		if ((TestFunc->GetFName() == DelegateSignatureFunction->GetFName()) && (TestFunc != DelegateSignatureFunction))
+		{
+			FError::Throwf(TEXT("Can't override delegate signature function '%s'"), FuncInfoFunctionIdentifier);
+		}
+	}
+
+	return DelegateSignatureFunction;
+}
+
+// Compares the properties of two functions to see if they have the same signature.
+bool AreFunctionSignaturesEqual(const UFunction* Lhs, const UFunction* Rhs)
+{
+	auto LhsPropIter = TFieldIterator<FProperty>(Lhs);
+	auto RhsPropIter = TFieldIterator<FProperty>(Rhs);
+
+	for (;;)
+	{
+		bool bEndOfLhsFunction = !LhsPropIter;
+		bool bEndOfRhsFunction = !RhsPropIter;
+
+		if (bEndOfLhsFunction != bEndOfRhsFunction)
+		{
+			// The functions have different numbers of parameters
+			return false;
+		}
+
+		if (bEndOfLhsFunction)
+		{
+			// We've compared all the parameters
+			return true;
+		}
+
+		const FProperty* LhsProp = *LhsPropIter;
+		const FProperty* RhsProp = *RhsPropIter;
+
+		FFieldClass* LhsClass = LhsProp->GetClass();
+		FFieldClass* RhsClass = RhsProp->GetClass();
+
+		if (LhsClass != RhsClass)
+		{
+			// The properties have different types
+			return false;
+		}
+
+		if (LhsClass == FArrayProperty::StaticClass())
+		{
+			const FArrayProperty* LhsArrayProp = (const FArrayProperty*)LhsProp;
+			const FArrayProperty* RhsArrayProp = (const FArrayProperty*)RhsProp;
+
+			if (LhsArrayProp->Inner->GetClass() != RhsArrayProp->Inner->GetClass())
+			{
+				// The properties are arrays of different types
+				return false;
+			}
+		}
+		else if (LhsClass == FMapProperty::StaticClass())
+		{
+			const FMapProperty* LhsMapProp = (const FMapProperty*)LhsProp;
+			const FMapProperty* RhsMapProp = (const FMapProperty*)RhsProp;
+
+			if (LhsMapProp->KeyProp->GetClass() != RhsMapProp->KeyProp->GetClass() || LhsMapProp->ValueProp->GetClass() != RhsMapProp->ValueProp->GetClass())
+			{
+				// The properties are maps of different types
+				return false;
+			}
+		}
+		else if (LhsClass == FSetProperty::StaticClass())
+		{
+			const FSetProperty* LhsSetProp = (const FSetProperty*)LhsProp;
+			const FSetProperty* RhsSetProp = (const FSetProperty*)RhsProp;
+
+			if (LhsSetProp->ElementProp->GetClass() != RhsSetProp->ElementProp->GetClass())
+			{
+			// The properties are sets of different types
+			return false;
+			}
+		}
+
+		++LhsPropIter;
+		++RhsPropIter;
+	}
+}
+
+/**
+ * Parses and compiles a function declaration
+ */
+void FHeaderParser::CompileFunctionDeclaration(FClasses& AllClasses)
+{
+	CheckAllow(TEXT("'Function'"), ENestAllowFlags::Function);
+
+	FUnrealSourceFile* CurrentSrcFile = GetCurrentSourceFile();
+	TMap<FName, FString> MetaData;
+	AddModuleRelativePathToMetadata(*CurrentSrcFile, MetaData);
+
+	// New-style UFUNCTION() syntax 
+	TArray<FPropertySpecifier> SpecifiersFound;
+	ReadSpecifierSetInsideMacro(SpecifiersFound, TEXT("Function"), MetaData);
+
+	FScriptLocation FuncNameRetry;
+	InitScriptLocation(FuncNameRetry);
+
+	if (!GetCurrentClass()->HasAnyClassFlags(CLASS_Native))
+	{
+		FError::Throwf(TEXT("Should only be here for native classes!"));
+	}
+
+	// Process all specifiers.
+	const TCHAR* TypeOfFunction = TEXT("function");
+
+	bool bAutomaticallyFinal = true;
+
+	FFuncInfo FuncInfo;
+	FuncInfo.MacroLine = InputLine;
+	FuncInfo.FunctionFlags = FUNC_Native;
+
+	// Infer the function's access level from the currently declared C++ access level
+	if (CurrentAccessSpecifier == ACCESS_Public)
+	{
+		FuncInfo.FunctionFlags |= FUNC_Public;
+	}
+	else if (CurrentAccessSpecifier == ACCESS_Protected)
+	{
+		FuncInfo.FunctionFlags |= FUNC_Protected;
+	}
+	else if (CurrentAccessSpecifier == ACCESS_Private)
+	{
+		FuncInfo.FunctionFlags |= FUNC_Private;
+		FuncInfo.FunctionFlags |= FUNC_Final;
+
+		// This is automatically final as well, but in a different way and for a different reason
+		bAutomaticallyFinal = false;
+	}
+	else
+	{
+		FError::Throwf(TEXT("Unknown access level"));
+	}
+
+	// non-static functions in a const class must be const themselves
+	if (GetCurrentClass()->HasAnyClassFlags(CLASS_Const))
+	{
+		FuncInfo.FunctionFlags |= FUNC_Const;
+	}
+
+	if (MatchIdentifier(TEXT("static"), ESearchCase::CaseSensitive))
+	{
+		FuncInfo.FunctionFlags |= FUNC_Static;
+		FuncInfo.FunctionExportFlags |= FUNCEXPORT_CppStatic;
+	}
+
+	if (MetaData.Contains(NAME_CppFromBpEvent))
+	{
+		FuncInfo.FunctionFlags |= FUNC_Event;
+	}
+
+	if (CompilerDirectiveStack.Num() > 0 && (CompilerDirectiveStack.Last()&ECompilerDirective::WithEditor) != 0)
+	{
+		FuncInfo.FunctionFlags |= FUNC_EditorOnly;
+	}
+
+	ProcessFunctionSpecifiers(FuncInfo, SpecifiersFound, MetaData);
+
+	const bool bClassGeneratedFromBP = FClass::IsDynamic(GetCurrentClass());
+
+	if ((0 != (FuncInfo.FunctionExportFlags & FUNCEXPORT_CustomThunk)) && !MetaData.Contains(NAME_CustomThunk))
+	{
+		MetaData.Add(NAME_CustomThunk, TEXT("true"));
+	}
+
+	if ((FuncInfo.FunctionFlags & FUNC_BlueprintPure) && GetCurrentClass()->HasAnyClassFlags(CLASS_Interface))
+	{
+		// Until pure interface casts are supported, we don't allow pures in interfaces
+		UE_LOG_ERROR_UHT(TEXT("BlueprintPure specifier is not allowed for interface functions"));
+	}
+
+	if (FuncInfo.FunctionFlags & FUNC_Net)
+	{
+		// Network replicated functions are always events, and are only final if sealed
+		TypeOfFunction = TEXT("event");
+		bAutomaticallyFinal = false;
+	}
+
+	if (FuncInfo.FunctionFlags & FUNC_BlueprintEvent)
+	{
+		TypeOfFunction = (FuncInfo.FunctionFlags & FUNC_Native) ? TEXT("BlueprintNativeEvent") : TEXT("BlueprintImplementableEvent");
+		bAutomaticallyFinal = false;
+	}
+
+	bool bSawVirtual = false;
+
+	if (MatchIdentifier(TEXT("virtual"), ESearchCase::CaseSensitive))
+	{
+		bSawVirtual = true;
+	}
+
+	FString*   InternalPtr = MetaData.Find(NAME_BlueprintInternalUseOnly); // FBlueprintMetadata::MD_BlueprintInternalUseOnly
+	const bool bInternalOnly = InternalPtr && *InternalPtr == TEXT("true");
+
+	// If this function is blueprint callable or blueprint pure, require a category 
+	if ((FuncInfo.FunctionFlags & (FUNC_BlueprintCallable | FUNC_BlueprintPure)) != 0) 
+	{ 
+		const bool bDeprecated = MetaData.Contains(NAME_DeprecatedFunction);       // FBlueprintMetadata::MD_DeprecatedFunction
+		const bool bBlueprintAccessor = MetaData.Contains(NAME_BlueprintSetter) || MetaData.Contains(NAME_BlueprintGetter); // FBlueprintMetadata::MD_BlueprintSetter, // FBlueprintMetadata::MD_BlueprintGetter
+		const bool bHasMenuCategory = MetaData.Contains(NAME_Category);                 // FBlueprintMetadata::MD_FunctionCategory
+
+		if (!bHasMenuCategory && !bInternalOnly && !bDeprecated && !bBlueprintAccessor) 
+		{ 
+			// To allow for quick iteration, don't enforce the requirement that game functions have to be categorized
+			if (bIsCurrentModulePartOfEngine)
+			{
+				UE_LOG_ERROR_UHT(TEXT("An explicit Category specifier is required for Blueprint accessible functions in an Engine module."));
+			}
+		}
+	}
+
+	// Verify interfaces with respect to their blueprint accessible functions
+	if (GetCurrentClass()->HasAnyClassFlags(CLASS_Interface))
+	{
+		if((FuncInfo.FunctionFlags & FUNC_BlueprintEvent) != 0 && !bInternalOnly)
+		{
+			const bool bCanImplementInBlueprints = !GetCurrentClass()->HasMetaData(NAME_CannotImplementInterfaceInBlueprint);  //FBlueprintMetadata::MD_CannotImplementInterfaceInBlueprint
+
+			// Ensure that blueprint events are only allowed in implementable interfaces. Internal only functions allowed
+			if (!bCanImplementInBlueprints)
+			{
+				UE_LOG_ERROR_UHT(TEXT("Interfaces that are not implementable in blueprints cannot have BlueprintImplementableEvent members."));
+			}
+		}
+		
+		if (((FuncInfo.FunctionFlags & FUNC_BlueprintCallable) != 0) && (((~FuncInfo.FunctionFlags) & FUNC_BlueprintEvent) != 0))
+		{
+			const bool bCanImplementInBlueprints = !GetCurrentClass()->HasMetaData(NAME_CannotImplementInterfaceInBlueprint);  //FBlueprintMetadata::MD_CannotImplementInterfaceInBlueprint
+
+			// Ensure that if this interface contains blueprint callable functions that are not blueprint defined, that it must be implemented natively
+			if (bCanImplementInBlueprints)
+			{
+				UE_LOG_ERROR_UHT(TEXT("Blueprint implementable interfaces cannot contain BlueprintCallable functions that are not BlueprintImplementableEvents.  Use CannotImplementInterfaceInBlueprint on the interface if you wish to keep this function."));
+			}
+		}
+	}
+
+	// Peek ahead to look for a CORE_API style DLL import/export token if present
+	FString APIMacroIfPresent;
+	{
+		FToken Token;
+		if (GetToken(Token, true))
+		{
+			bool bThrowTokenBack = true;
+			if (Token.TokenType == TOKEN_Identifier)
+			{
+				FString RequiredAPIMacroIfPresent(Token.Identifier);
+				if (RequiredAPIMacroIfPresent.EndsWith(TEXT("_API"), ESearchCase::CaseSensitive))
+				{
+					//@TODO: Validate the module name for RequiredAPIMacroIfPresent
+					bThrowTokenBack = false;
+
+					if (GetCurrentClass()->HasAnyClassFlags(CLASS_RequiredAPI))
+					{
+						FError::Throwf(TEXT("'%s' must not be used on methods of a class that is marked '%s' itself."), *RequiredAPIMacroIfPresent, *RequiredAPIMacroIfPresent);
+					}
+					FuncInfo.FunctionFlags |= FUNC_RequiredAPI;
+					FuncInfo.FunctionExportFlags |= FUNCEXPORT_RequiredAPI;
+
+					APIMacroIfPresent = RequiredAPIMacroIfPresent;
+				}
+			}
+
+			if (bThrowTokenBack)
+			{
+				UngetToken(Token);
+			}
+		}
+	}
+
+	// Look for static again, in case there was an ENGINE_API token first
+	if (!APIMacroIfPresent.IsEmpty() && MatchIdentifier(TEXT("static"), ESearchCase::CaseSensitive))
+	{
+		FError::Throwf(TEXT("Unexpected API macro '%s'. Did you mean to put '%s' after the static keyword?"), *APIMacroIfPresent, *APIMacroIfPresent);
+	}
+
+	// Look for virtual again, in case there was an ENGINE_API token first
+	if (MatchIdentifier(TEXT("virtual"), ESearchCase::CaseSensitive))
+	{
+		bSawVirtual = true;
+	}
+
+	// Process the virtualness
+	if (bSawVirtual)
+	{
+		// Remove the implicit final, the user can still specifying an explicit final at the end of the declaration
+		bAutomaticallyFinal = false;
+
+		// if this is a BlueprintNativeEvent or BlueprintImplementableEvent in an interface, make sure it's not "virtual"
+		if (FuncInfo.FunctionFlags & FUNC_BlueprintEvent)
+		{
+			if (GetCurrentClass()->HasAnyClassFlags(CLASS_Interface))
+			{
+				FError::Throwf(TEXT("BlueprintImplementableEvents in Interfaces must not be declared 'virtual'"));
+			}
+
+			// if this is a BlueprintNativeEvent, make sure it's not "virtual"
+			else if (FuncInfo.FunctionFlags & FUNC_Native)
+			{
+				UE_LOG_ERROR_UHT(TEXT("BlueprintNativeEvent functions must be non-virtual."));
+			}
+
+			else
+			{
+				UE_LOG_WARNING_UHT(TEXT("BlueprintImplementableEvents should not be virtual. Use BlueprintNativeEvent instead."));
+			}
+		}
+	}
+	else
+	{
+		// if this is a function in an Interface, it must be marked 'virtual' unless it's an event
+		if (GetCurrentClass()->HasAnyClassFlags(CLASS_Interface) && !(FuncInfo.FunctionFlags & FUNC_BlueprintEvent))
+		{
+			FError::Throwf(TEXT("Interface functions that are not BlueprintImplementableEvents must be declared 'virtual'"));
+		}
+	}
+
+	// Handle the initial implicit/explicit final
+	// A user can still specify an explicit final after the parameter list as well.
+	if (bAutomaticallyFinal || FuncInfo.bSealedEvent)
+	{
+		FuncInfo.FunctionFlags |= FUNC_Final;
+		FuncInfo.FunctionExportFlags |= FUNCEXPORT_Final;
+
+		if (GetCurrentClass()->HasAnyClassFlags(CLASS_Interface))
+		{
+			UE_LOG_ERROR_UHT(TEXT("Interface functions cannot be declared 'final'"));
+		}
+	}
+
+	// Get return type.
+	FToken ReturnType( CPT_None );
+
+	// C++ style functions always have a return value type, even if it's void
+	bool bHasReturnValue = !MatchIdentifier(TEXT("void"), ESearchCase::CaseSensitive);
+	if (bHasReturnValue)
+	{
+		GetVarType(AllClasses, GetCurrentScope(), ReturnType, CPF_None, nullptr, EPropertyDeclarationStyle::None, EVariableCategory::Return);
+	}
+
+	// Skip whitespaces to get InputPos exactly on beginning of function name.
+	while (FChar::IsWhitespace(PeekChar())) { GetChar(); }
+
+	FuncInfo.InputPos = InputPos;
+
+	// Get function or operator name.
+	if (!GetIdentifier(FuncInfo.Function))
+	{
+		FError::Throwf(TEXT("Missing %s name"), TypeOfFunction);
+	}
+
+	if ( !MatchSymbol(TEXT('(')) )
+	{
+		FError::Throwf(TEXT("Bad %s definition"), TypeOfFunction);
+	}
+
+	if (FuncInfo.FunctionFlags & FUNC_Net)
+	{
+		bool bIsNetService = !!(FuncInfo.FunctionFlags & (FUNC_NetRequest | FUNC_NetResponse));
+		if (bHasReturnValue && !bIsNetService)
+		{
+			FError::Throwf(TEXT("Replicated functions can't have return values"));
+		}
+
+		if (FuncInfo.RPCId > 0)
+		{
+			if (FString* ExistingFunc = UsedRPCIds.Find(FuncInfo.RPCId))
+			{
+				FError::Throwf(TEXT("Function %s already uses identifier %d"), **ExistingFunc, FuncInfo.RPCId);
+			}
+
+			UsedRPCIds.Add(FuncInfo.RPCId, FuncInfo.Function.Identifier);
+			if (FuncInfo.FunctionFlags & FUNC_NetResponse)
+			{
+				// Look for another function expecting this response
+				if (FString* ExistingFunc = RPCsNeedingHookup.Find(FuncInfo.RPCId))
+				{
+					// If this list isn't empty at end of class, throw error
+					RPCsNeedingHookup.Remove(FuncInfo.RPCId);
+				}
+			}
+		}
+
+		if (FuncInfo.RPCResponseId > 0 && FuncInfo.EndpointName != TEXT("JSBridge"))
+		{
+			// Look for an existing response function
+			FString* ExistingFunc = UsedRPCIds.Find(FuncInfo.RPCResponseId);
+			if (ExistingFunc == NULL)
+			{
+				// If this list isn't empty at end of class, throw error
+				RPCsNeedingHookup.Add(FuncInfo.RPCResponseId, FuncInfo.Function.Identifier);
+			}
+		}
+	}
+
+	UFunction* TopFunction = CreateFunction(FuncInfo);
+
+	FClassMetaData* ClassMetaData = GScriptHelper.AddClassData(TopFunction, CurrentSrcFile);
+
+	TopFunction->FunctionFlags |= FuncInfo.FunctionFlags;
+
+	FuncInfo.FunctionReference = TopFunction;
+	FuncInfo.SetFunctionNames();
+
+	GetCurrentScope()->AddType(TopFunction);
+
+	FFunctionData* StoredFuncData = FFunctionData::Add(FFuncInfo(FuncInfo));
+	if (FuncInfo.FunctionReference->HasAnyFunctionFlags(FUNC_Delegate))
+	{
+		GetCurrentClassData()->MarkContainsDelegate();
+	}
+
+	// Get parameter list.
+	ParseParameterList(AllClasses, TopFunction, false, &MetaData);
+
+	// Get return type, if any.
+	if (bHasReturnValue)
+	{
+		ReturnType.PropertyFlags |= CPF_Parm | CPF_OutParm | CPF_ReturnParm;
+		FProperty* ReturnProp = GetVarNameAndDim(TopFunction, ReturnType, EVariableCategory::Return);
+
+		TopFunction->NumParms++;
+	}
+
+	// determine if there are any outputs for this function
+	bool bHasAnyOutputs = bHasReturnValue;
+	if (!bHasAnyOutputs)
+	{
+		for (TFieldIterator<FProperty> It(TopFunction); It; ++It)
+		{
+			FProperty const* const Param = *It;
+			if (!(Param->PropertyFlags & CPF_ReturnParm) && (Param->PropertyFlags & CPF_OutParm))
+			{
+				bHasAnyOutputs = true;
+				break;
+			}
+		}
+	}
+
+	// Check to see if there is a function in the super class with the same name
+	UStruct* SuperStruct = GetCurrentClass();
+	if (SuperStruct)
+	{
+		SuperStruct = SuperStruct->GetSuperStruct();
+	}
+	if (SuperStruct)
+	{
+		if (UFunction* OverriddenFunction = ::FindUField<UFunction>(SuperStruct, FuncInfo.Function.Identifier))
+		{
+			// Native function overrides should be done in CPP text, not in a UFUNCTION() declaration (you can't change flags, and it'd otherwise be a burden to keep them identical)
+			UE_LOG_ERROR_UHT(TEXT("%s: Override of UFUNCTION in parent class (%s) cannot have a UFUNCTION() declaration above it; it will use the same parameters as the original declaration."), FuncInfo.Function.Identifier, *OverriddenFunction->GetOuter()->GetName());
+		}
+	}
+
+	if (!bHasAnyOutputs && (FuncInfo.FunctionFlags & (FUNC_BlueprintPure)))
+	{
+		// This bad behavior would be treated as a warning in the Blueprint editor, so when converted assets generates these bad functions
+		// we don't want to prevent compilation:
+		if (!bClassGeneratedFromBP)
+		{
+			UE_LOG_ERROR_UHT(TEXT("BlueprintPure specifier is not allowed for functions with no return value and no output parameters."));
+		}
+	}
+
+
+	// determine whether this function should be 'const'
+	if ( MatchIdentifier(TEXT("const"), ESearchCase::CaseSensitive) )
+	{
+		if( (TopFunction->FunctionFlags & (FUNC_Native)) == 0 )
+		{
+			// @TODO: UCREMOVAL Reconsider?
+			//FError::Throwf(TEXT("'const' may only be used for native functions"));
+		}
+
+		FuncInfo.FunctionFlags |= FUNC_Const;
+
+		// @todo: the presence of const and one or more outputs does not guarantee that there are
+		// no side effects. On GCC and clang we could use __attribure__((pure)) or __attribute__((const))
+		// or we could just rely on the use marking things BlueprintPure. Either way, checking the C++
+		// const identifier to determine purity is not desirable. We should remove the following logic:
+
+		// If its a const BlueprintCallable function with some sort of output and is not being marked as an BlueprintPure=false function, mark it as BlueprintPure as well
+		if ( bHasAnyOutputs && ((FuncInfo.FunctionFlags & FUNC_BlueprintCallable) != 0) && !FuncInfo.bForceBlueprintImpure)
+		{
+			FuncInfo.FunctionFlags |= FUNC_BlueprintPure;
+		}
+	}
+
+	// Try parsing metadata for the function
+	ParseFieldMetaData(MetaData, *(TopFunction->GetName()));
+
+	AddFormattedPrevCommentAsTooltipMetaData(MetaData);
+
+	AddMetaDataToClassData(TopFunction, MoveTemp(MetaData));
+
+	// 'final' and 'override' can appear in any order before an optional '= 0' pure virtual specifier
+	bool bFoundFinal    = MatchIdentifier(TEXT("final"), ESearchCase::CaseSensitive);
+	bool bFoundOverride = MatchIdentifier(TEXT("override"), ESearchCase::CaseSensitive);
+	if (!bFoundFinal && bFoundOverride)
+	{
+		bFoundFinal = MatchIdentifier(TEXT("final"), ESearchCase::CaseSensitive);
+	}
+
+	// Handle C++ style functions being declared as abstract
+	if (MatchSymbol(TEXT('=')))
+	{
+		int32 ZeroValue = 1;
+		bool bGotZero = GetConstInt(/*out*/ZeroValue);
+		bGotZero = bGotZero && (ZeroValue == 0);
+
+		if (!bGotZero)
+		{
+			FError::Throwf(TEXT("Expected 0 to indicate function is abstract"));
+		}
+	}
+
+	// Look for the final keyword to indicate this function is sealed
+	if (bFoundFinal)
+	{
+		// This is a final (prebinding, non-overridable) function
+		FuncInfo.FunctionFlags |= FUNC_Final;
+		FuncInfo.FunctionExportFlags |= FUNCEXPORT_Final;
+		if (GetCurrentClass()->HasAnyClassFlags(CLASS_Interface))
+		{
+			FError::Throwf(TEXT("Interface functions cannot be declared 'final'"));
+		}
+		else if (FuncInfo.FunctionFlags & FUNC_BlueprintEvent)
+		{
+			FError::Throwf(TEXT("Blueprint events cannot be declared 'final'"));
+		}
+	}
+
+	// Make sure any new flags made it to the function
+	//@TODO: UCREMOVAL: Ideally the flags didn't get copied midway thru parsing the function declaration, and we could avoid this
+	TopFunction->FunctionFlags |= FuncInfo.FunctionFlags;
+	StoredFuncData->UpdateFunctionData(FuncInfo);
+
+	// Bind the function.
+	TopFunction->Bind();
+	
+	// Make sure that the replication flags set on an overridden function match the parent function
+	if (UFunction* SuperFunc = TopFunction->GetSuperFunction())
+	{
+		if ((TopFunction->FunctionFlags & FUNC_NetFuncFlags) != (SuperFunc->FunctionFlags & FUNC_NetFuncFlags))
+		{
+			FError::Throwf(TEXT("Overridden function '%s': Cannot specify different replication flags when overriding a function."), *TopFunction->GetName());
+		}
+	}
+
+	// if this function is an RPC in state scope, verify that it is an override
+	// this is required because the networking code only checks the class for RPCs when initializing network data, not any states within it
+	if ((TopFunction->FunctionFlags & FUNC_Net) && (TopFunction->GetSuperFunction() == NULL) && Cast<UClass>(TopFunction->GetOuter()) == NULL)
+	{
+		FError::Throwf(TEXT("Function '%s': Base implementation of RPCs cannot be in a state. Add a stub outside state scope."), *TopFunction->GetName());
+	}
+
+	if (TopFunction->FunctionFlags & (FUNC_BlueprintCallable | FUNC_BlueprintEvent))
+	{
+		for (TFieldIterator<FProperty> It(TopFunction); It; ++It)
+		{
+			FProperty const* const Param = *It;
+			if (Param->ArrayDim > 1)
+			{
+				FError::Throwf(TEXT("Static array cannot be exposed to blueprint. Function: %s Parameter %s\n"), *TopFunction->GetName(), *Param->GetName());
+			}
+
+			if (!IsPropertySupportedByBlueprint(Param, false))
+			{
+				FString ExtendedCPPType;
+				FString CPPType = Param->GetCPPType(&ExtendedCPPType);
+				UE_LOG_ERROR_UHT(TEXT("Type '%s%s' is not supported by blueprint. %s.%s"), *CPPType, *ExtendedCPPType, *TopFunction->GetName(), *Param->GetName());
+			}
+		}
+	}
+
+	// Just declaring a function, so end the nesting.
+	PostPopFunctionDeclaration(AllClasses, TopFunction);
+
+	// See what's coming next
+	FToken Token;
+	if (!GetToken(Token))
+	{
+		FError::Throwf(TEXT("Unexpected end of file"));
+	}
+
+	// Optionally consume a semicolon
+	// This is optional to allow inline function definitions
+	if (Token.TokenType == TOKEN_Symbol && Token.Matches(TEXT(';')))
+	{
+		// Do nothing (consume it)
+	}
+	else if (Token.TokenType == TOKEN_Symbol && Token.Matches(TEXT('{')))
+	{
+		// Skip inline function bodies
+		UngetToken(Token);
+		SkipDeclaration(Token);
+	}
+	else
+	{
+		// Put the token back so we can continue parsing as normal
+		UngetToken(Token);
+	}
+
+	// perform documentation policy tests
+	CheckDocumentationPolicyForFunc(GetCurrentClass(), FuncInfo.FunctionReference);
+}
+
+/** Parses optional metadata text. */
+void FHeaderParser::ParseFieldMetaData(TMap<FName, FString>& MetaData, const TCHAR* FieldName)
+{
+	FToken PropertyMetaData;
+	bool bMetadataPresent = false;
+	if (MatchIdentifier(TEXT("UMETA"), ESearchCase::CaseSensitive))
+	{
+		auto ErrorMessageGetter = [FieldName]() { return FString::Printf(TEXT("' %s metadata'"), FieldName); };
+
+		bMetadataPresent = true;
+		RequireSymbol( TEXT('('), ErrorMessageGetter );
+		if (!GetRawTokenRespectingQuotes(PropertyMetaData, TCHAR(')')))
+		{
+			FError::Throwf(TEXT("'%s': No metadata specified"), FieldName);
+		}
+		RequireSymbol( TEXT(')'), ErrorMessageGetter);
+	}
+
+	if (bMetadataPresent)
+	{
+		// parse apart the string
+		TArray<FString> Pairs;
+
+		//@TODO: UCREMOVAL: Convert to property token reading
+		// break apart on | to get to the key/value pairs
+		FString NewData(PropertyMetaData.String);
+		bool bInString = false;
+		int32 LastStartIndex = 0;
+		int32 CharIndex;
+		for (CharIndex = 0; CharIndex < NewData.Len(); ++CharIndex)
+		{
+			TCHAR Ch = NewData.GetCharArray()[CharIndex];
+			if (Ch == '"')
+			{
+				bInString = !bInString;
+			}
+
+			if ((Ch == ',') && !bInString)
+			{
+				if (LastStartIndex != CharIndex)
+				{
+					Pairs.Add(NewData.Mid(LastStartIndex, CharIndex - LastStartIndex));
+				}
+				LastStartIndex = CharIndex + 1;
+			}
+		}
+
+		if (LastStartIndex != CharIndex)
+		{
+			Pairs.Add(MoveTemp(NewData).Mid(LastStartIndex, CharIndex - LastStartIndex));
+		}
+
+		// go over all pairs
+		for (int32 PairIndex = 0; PairIndex < Pairs.Num(); PairIndex++)
+		{
+			// break the pair into a key and a value
+			FString Token = MoveTemp(Pairs[PairIndex]);
+			FString Key;
+			// by default, not value, just a key (allowed)
+			FString Value;
+
+			// look for a value after an =
+			const int32 Equals = Token.Find(TEXT("="), ESearchCase::CaseSensitive);
+			// if we have an =, break up the string
+			if (Equals != INDEX_NONE)
+			{
+				Key = Token.Left(Equals);
+				Value = MoveTemp(Token);
+				Value.RightInline((Value.Len() - Equals) - 1, false);
+			}
+			else
+			{
+				Key = MoveTemp(Token);
+			}
+
+			InsertMetaDataPair(MetaData, MoveTemp(Key), MoveTemp(Value));
+		}
+	}
+}
+
+bool FHeaderParser::IsBitfieldProperty(ELayoutMacroType LayoutMacroType)
+{
+	if (LayoutMacroType == ELayoutMacroType::Bitfield || LayoutMacroType == ELayoutMacroType::BitfieldEditorOnly)
+	{
+		return true;
+	}
+
+	bool bIsBitfield = false;
+
+	// The current token is the property type (uin32, uint16, etc).
+	// Check the property name and then check for ':'
+	FToken TokenVarName;
+	if (GetToken(TokenVarName, /*bNoConsts=*/ true))
+	{
+		FToken Token;
+		if (GetToken(Token, /*bNoConsts=*/ true))
+		{
+			if (Token.TokenType == TOKEN_Symbol && Token.Matches(TEXT(':')))
+			{
+				bIsBitfield = true;
+			}
+			UngetToken(Token);
+		}
+		UngetToken(TokenVarName);
+	}
+
+	return bIsBitfield;
+}
+
+void FHeaderParser::ValidatePropertyIsDeprecatedIfNecessary(const FPropertyBase& VarProperty, const FToken* OuterPropertyType)
+{
+	// check to see if we have a FClassProperty using a deprecated class
+	if ( VarProperty.MetaClass != NULL && VarProperty.MetaClass->HasAnyClassFlags(CLASS_Deprecated) && !(VarProperty.PropertyFlags & CPF_Deprecated) &&
+		(OuterPropertyType == NULL || !(OuterPropertyType->PropertyFlags & CPF_Deprecated)) )
+	{
+		UE_LOG_ERROR_UHT(TEXT("Property is using a deprecated class: %s.  Property should be marked deprecated as well."), *VarProperty.MetaClass->GetPathName());
+	}
+
+	// check to see if we have a FObjectProperty using a deprecated class.
+	// PropertyClass is part of a union, so only check PropertyClass if this token represents an object property
+	if ( (VarProperty.Type == CPT_ObjectReference || VarProperty.Type == CPT_WeakObjectReference || VarProperty.Type == CPT_LazyObjectReference || VarProperty.Type == CPT_SoftObjectReference) && VarProperty.PropertyClass != NULL
+		&&	VarProperty.PropertyClass->HasAnyClassFlags(CLASS_Deprecated)	// and the object class being used has been deprecated
+		&& (VarProperty.PropertyFlags&CPF_Deprecated) == 0					// and this property isn't marked deprecated as well
+		&& (OuterPropertyType == NULL || !(OuterPropertyType->PropertyFlags & CPF_Deprecated)) ) // and this property isn't in an array that was marked deprecated either
+	{
+		UE_LOG_ERROR_UHT(TEXT("Property is using a deprecated class: %s.  Property should be marked deprecated as well."), *VarProperty.PropertyClass->GetPathName());
+	}
+}
+
+struct FExposeOnSpawnValidator
+{
+	// Keep this function synced with UEdGraphSchema_K2::FindSetVariableByNameFunction
+	static bool IsSupported(const FPropertyBase& Property)
+	{
+		bool ProperNativeType = false;
+		switch (Property.Type)
+		{
+		case CPT_Int:
+		case CPT_Int64:
+		case CPT_Byte:
+		case CPT_Float:
+		case CPT_Bool:
+		case CPT_Bool8:
+		case CPT_ObjectReference:
+		case CPT_String:
+		case CPT_Text:
+		case CPT_Name:
+		case CPT_Interface:
+		case CPT_SoftObjectReference:
+			ProperNativeType = true;
+		}
+
+		if (!ProperNativeType && (CPT_Struct == Property.Type) && Property.Struct)
+		{
+			ProperNativeType |= Property.Struct->GetBoolMetaData(NAME_BlueprintType);
+		}
+
+		return ProperNativeType;
+	}
+};
+
+void FHeaderParser::CompileVariableDeclaration(FClasses& AllClasses, UStruct* Struct)
+{
+	EPropertyFlags DisallowFlags = CPF_ParmFlags;
+	EPropertyFlags EdFlags       = CPF_None;
+
+	// Get variable type.
+	FPropertyBase OriginalProperty(CPT_None);
+	FIndexRange TypeRange;
+	ELayoutMacroType LayoutMacroType = ELayoutMacroType::None;
+	GetVarType( AllClasses, &FScope::GetTypeScope(Struct).Get(), OriginalProperty, DisallowFlags, /*OuterPropertyType=*/ NULL, EPropertyDeclarationStyle::UPROPERTY, EVariableCategory::Member, &TypeRange, &LayoutMacroType);
+	OriginalProperty.PropertyFlags |= EdFlags;
+
+	FString* Category = OriginalProperty.MetaData.Find(NAME_Category);
+
+	// First check if the category was specified at all and if the property was exposed to the editor.
+	if (!Category && (OriginalProperty.PropertyFlags & (CPF_Edit|CPF_BlueprintVisible)))
+	{
+		if ((Struct->GetOutermost() != nullptr) && !bIsCurrentModulePartOfEngine)
+		{
+			Category = &OriginalProperty.MetaData.Add(NAME_Category, Struct->GetName());
+		}
+		else
+		{
+			UE_LOG_ERROR_UHT(TEXT("An explicit Category specifier is required for any property exposed to the editor or Blueprints in an Engine module."));
+		}
+	}
+
+	// Validate that pointer properties are not interfaces (which are not GC'd and so will cause runtime errors)
+	if (OriginalProperty.PointerType == EPointerType::Native && OriginalProperty.Struct->IsChildOf(UInterface::StaticClass()))
+	{
+		// Get the name of the type, removing the asterisk representing the pointer
+		FString TypeName = FString(TypeRange.Count, Input + TypeRange.StartIndex).TrimStartAndEnd().LeftChop(1).TrimEnd();
+		FError::Throwf(TEXT("UPROPERTY pointers cannot be interfaces - did you mean TScriptInterface<%s>?"), *TypeName);
+	}
+
+	// If the category was specified explicitly, it wins
+	if (Category && !(OriginalProperty.PropertyFlags & (CPF_Edit|CPF_BlueprintVisible|CPF_BlueprintAssignable|CPF_BlueprintCallable)))
+	{
+		UE_LOG_WARNING_UHT(TEXT("Property has a Category set but is not exposed to the editor or Blueprints with EditAnywhere, BlueprintReadWrite, VisibleAnywhere, BlueprintReadOnly, BlueprintAssignable, BlueprintCallable keywords.\r\n"));
+	}
+
+	// Make sure that editblueprint variables are editable
+	if(!(OriginalProperty.PropertyFlags & CPF_Edit))
+	{
+		if (OriginalProperty.PropertyFlags & CPF_DisableEditOnInstance)
+		{
+			UE_LOG_ERROR_UHT(TEXT("Property cannot have 'DisableEditOnInstance' without being editable"));
+		}
+
+		if (OriginalProperty.PropertyFlags & CPF_DisableEditOnTemplate)
+		{
+			UE_LOG_ERROR_UHT(TEXT("Property cannot have 'DisableEditOnTemplate' without being editable"));
+		}
+	}
+
+	// Validate.
+	if (OriginalProperty.PropertyFlags & CPF_ParmFlags)
+	{
+		FError::Throwf(TEXT("Illegal type modifiers in member variable declaration") );
+	}
+
+	if (FString* ExposeOnSpawnValue = OriginalProperty.MetaData.Find(NAME_ExposeOnSpawn))
+	{
+		if ((*ExposeOnSpawnValue == TEXT("true")) && !FExposeOnSpawnValidator::IsSupported(OriginalProperty))
+		{
+			UE_LOG_ERROR_UHT(TEXT("ExposeOnSpawn - Property cannot be exposed"));
+		}
+	}
+
+	if (LayoutMacroType != ELayoutMacroType::None)
+	{
+		RequireSymbol(TEXT(','), GLayoutMacroNames[(int32)LayoutMacroType]);
+	}
+
+	// Process all variables of this type.
+	TArray<FProperty*> NewProperties;
+	for (;;)
+	{
+		FToken     Property    = OriginalProperty;
+		FProperty* NewProperty = GetVarNameAndDim(Struct, Property, EVariableCategory::Member, LayoutMacroType);
+
+		// Optionally consume the :1 at the end of a bitfield boolean declaration
+		if (Property.IsBool())
+		{
+			if (LayoutMacroType == ELayoutMacroType::Bitfield || LayoutMacroType == ELayoutMacroType::BitfieldEditorOnly || MatchSymbol(TEXT(':')))
+			{
+				int32 BitfieldSize = 0;
+				if (!GetConstInt(/*out*/ BitfieldSize) || (BitfieldSize != 1))
+				{
+					FError::Throwf(TEXT("Bad or missing bitfield size for '%s', must be 1."), *NewProperty->GetName());
+				}
+			}
+		}
+
+		// Deprecation validation
+		ValidatePropertyIsDeprecatedIfNecessary(Property, NULL);
+
+		if (TopNest->NestType != ENestType::FunctionDeclaration)
+		{
+			if (NewProperties.Num())
+			{
+				FError::Throwf(TEXT("Comma delimited properties cannot be converted %s.%s\n"), *Struct->GetName(), *NewProperty->GetName());
+			}
+		}
+
+		NewProperties.Add( NewProperty );
+		// we'll need any metadata tags we parsed later on when we call ConvertEOLCommentToTooltip() so the tags aren't clobbered
+		OriginalProperty.MetaData = Property.MetaData;
+
+		if (NewProperty->HasAnyPropertyFlags(CPF_RepNotify))
+		{
+			NewProperty->RepNotifyFunc = OriginalProperty.RepNotifyName;
+		}
+
+		if (UScriptStruct* StructBeingBuilt = Cast<UScriptStruct>(Struct))
+		{
+			if (NewProperty->ContainsInstancedObjectProperty())
+			{
+				StructBeingBuilt->StructFlags = EStructFlags(StructBeingBuilt->StructFlags | STRUCT_HasInstancedReference);
+			}
+		}
+
+		if (NewProperty->HasAnyPropertyFlags(CPF_BlueprintVisible))
+		{
+			if (Struct->IsA<UScriptStruct>() && !Struct->GetBoolMetaDataHierarchical(NAME_BlueprintType))
+			{
+				UE_LOG_ERROR_UHT(TEXT("Cannot expose property to blueprints in a struct that is not a BlueprintType. %s.%s"), *Struct->GetName(), *NewProperty->GetName());
+			}
+
+			if (NewProperty->ArrayDim > 1)
+			{
+				UE_LOG_ERROR_UHT(TEXT("Static array cannot be exposed to blueprint %s.%s"), *Struct->GetName(), *NewProperty->GetName());
+			}
+
+			if (!IsPropertySupportedByBlueprint(NewProperty, true))
+			{
+				FString ExtendedCPPType;
+				FString CPPType = NewProperty->GetCPPType(&ExtendedCPPType);
+				UE_LOG_ERROR_UHT(TEXT("Type '%s%s' is not supported by blueprint. %s.%s"), *CPPType, *ExtendedCPPType, *Struct->GetName(), *NewProperty->GetName());
+			}
+		}
+
+		if (LayoutMacroType != ELayoutMacroType::None || !MatchSymbol(TEXT(',')))
+		{
+			break;
+		}
+	}
+
+	// Optional member initializer.
+	if (LayoutMacroType == ELayoutMacroType::FieldInitialized)
+	{
+		// Skip past the specified member initializer; we make no attempt to parse it
+		FToken SkipToken;
+		int Nesting = 1;
+		while (GetToken(SkipToken))
+		{
+			if (SkipToken.Matches(TEXT('(')))
+			{
+				++Nesting;
+			}
+			else if (SkipToken.Matches(TEXT(')')))
+			{
+				--Nesting;
+				if (Nesting == 0)
+				{
+					UngetToken(SkipToken);
+					break;
+				}
+			}
+		}
+	}
+	else if (MatchSymbol(TEXT('=')))
+	{
+		// Skip past the specified member initializer; we make no attempt to parse it
+		FToken SkipToken;
+		while (GetToken(SkipToken))
+		{
+			if (SkipToken.Matches(TEXT(';')))
+			{
+				// went too far
+				UngetToken(SkipToken);
+				break;
+			}
+		}
+	}
+	// Using Brace Initialization
+	else if (MatchSymbol(TEXT('{')))
+	{
+		FToken SkipToken;
+		int BraceLevel = 1;
+		while (GetToken(SkipToken))
+		{
+			if (SkipToken.Matches(TEXT('{')))
+			{
+				++BraceLevel;
+			}
+			else if (SkipToken.Matches(TEXT('}')))
+			{
+				--BraceLevel;
+				if (BraceLevel == 0)
+				{
+					break;
+				}
+			}
+		}
+	}
+
+	if (LayoutMacroType == ELayoutMacroType::None)
+	{
+		// Expect a semicolon.
+		RequireSymbol(TEXT(';'), TEXT("'variable declaration'"));
+	}
+	else
+	{
+		// Expect a close bracket.
+		RequireSymbol(TEXT(')'), GLayoutMacroNames[(int32)LayoutMacroType]);
+	}
+
+	// Skip redundant semi-colons
+	for (;;)
+	{
+		int32 CurrInputPos  = InputPos;
+		int32 CurrInputLine = InputLine;
+
+		FToken Token;
+		if (!GetToken(Token, /*bNoConsts=*/ true))
+		{
+			break;
+		}
+
+		if (Token.TokenType != TOKEN_Symbol || !Token.Matches(TEXT(';')))
+		{
+			InputPos  = CurrInputPos;
+			InputLine = CurrInputLine;
+			break;
+		}
+	}
+}
+
+//
+// Compile a statement: Either a declaration or a command.
+// Returns 1 if success, 0 if end of file.
+//
+bool FHeaderParser::CompileStatement(FClasses& AllClasses, TArray<UDelegateFunction*>& DelegatesToFixup)
+{
+	// Get a token and compile it.
+	FToken Token;
+	if( !GetToken(Token, true) )
+	{
+		// End of file.
+		return false;
+	}
+	else if (!CompileDeclaration(AllClasses, DelegatesToFixup, Token))
+	{
+		FError::Throwf(TEXT("'%s': Bad command or expression"), Token.Identifier );
+	}
+	return true;
+}
+
+//
+// Compute the function parameter size and save the return offset
+//
+//@TODO: UCREMOVAL: Need to rename ComputeFunctionParametersSize to reflect the additional work it's doing
+void FHeaderParser::ComputeFunctionParametersSize( UClass* Class )
+{
+	// Recurse with all child states in this class.
+	for (TFieldIterator<UFunction> FuncIt(Class, EFieldIteratorFlags::ExcludeSuper); FuncIt; ++FuncIt)
+	{
+		UFunction* ThisFunction = *FuncIt;
+
+		// Fix up any structs that were used as a parameter in a delegate before being defined
+		if (ThisFunction->HasAnyFunctionFlags(FUNC_Delegate))
+		{
+			for (TFieldIterator<FProperty> It(ThisFunction); It; ++It)
+			{
+				FProperty* Param = *It;
+				if (FStructProperty* StructProp = CastField<FStructProperty>(Param))
+				{
+					if (StructProp->Struct->StructFlags & STRUCT_HasInstancedReference)
+					{
+						StructProp->PropertyFlags |= CPF_ContainsInstancedReference;
+					}
+				}
+			}
+			ThisFunction->StaticLink(true);
+		}
+
+		// Compute the function parameter size, propagate some flags to the outer function, and save the return offset
+		// Must be done in a second phase, as StaticLink resets various fields again!
+		ThisFunction->ParmsSize = 0;
+		for (TFieldIterator<FProperty> It(ThisFunction); It; ++It)
+		{
+			FProperty* Param = *It;
+
+			if (!(Param->PropertyFlags & CPF_ReturnParm) && (Param->PropertyFlags & CPF_OutParm))
+			{
+				ThisFunction->FunctionFlags |= FUNC_HasOutParms;
+			}
+				
+			if (FStructProperty* StructProp = CastField<FStructProperty>(Param))
+			{
+				if (StructProp->Struct->HasDefaults())
+				{
+					ThisFunction->FunctionFlags |= FUNC_HasDefaults;
+				}
+			}
+		}
+	}
+}
+
+/*-----------------------------------------------------------------------------
+	Code skipping.
+-----------------------------------------------------------------------------*/
+
+/**
+ * Skip over code, honoring { and } pairs.
+ *
+ * @param	NestCount	number of nest levels to consume. if 0, consumes a single statement
+ * @param	ErrorTag	text to use in error message if EOF is encountered before we've done
+ */
+void FHeaderParser::SkipStatements( int32 NestCount, const TCHAR* ErrorTag  )
+{
+	FToken Token;
+
+	int32 OriginalNestCount = NestCount;
+
+	while( GetToken( Token, true ) )
+	{
+		if ( Token.Matches(TEXT('{')) )
+		{
+			NestCount++;
+		}
+		else if	( Token.Matches(TEXT('}')) )
+		{
+			NestCount--;
+		}
+		else if ( Token.Matches(TEXT(';')) && OriginalNestCount == 0 )
+		{
+			break;
+		}
+
+		if ( NestCount < OriginalNestCount || NestCount < 0 )
+			break;
+	}
+
+	if( NestCount > 0 )
+	{
+		FError::Throwf(TEXT("Unexpected end of file at end of %s"), ErrorTag );
+	}
+	else if ( NestCount < 0 )
+	{
+		FError::Throwf(TEXT("Extraneous closing brace found in %s"), ErrorTag);
+	}
+}
+
+/*-----------------------------------------------------------------------------
+	Main script compiling routine.
+-----------------------------------------------------------------------------*/
+
+//
+// Finalize any script-exposed functions in the specified class
+//
+void FHeaderParser::FinalizeScriptExposedFunctions(UClass* Class)
+{
+	// Finalize all of the children introduced in this class
+	for (TFieldIterator<UStruct> ChildIt(Class, EFieldIteratorFlags::ExcludeSuper); ChildIt; ++ChildIt)
+	{
+		UStruct* ChildStruct = *ChildIt;
+
+		if (UFunction* Function = Cast<UFunction>(ChildStruct))
+		{
+			// Add this function to the function map of its parent class
+			Class->AddFunctionToFunctionMap(Function, Function->GetFName());
+		}
+		else if (ChildStruct->IsA(UScriptStruct::StaticClass()))
+		{
+			// Ignore embedded structs
+		}
+		else
+		{
+			UE_LOG_WARNING_UHT(TEXT("Unknown and unexpected child named %s of type %s in %s\n"), *ChildStruct->GetName(), *ChildStruct->GetClass()->GetName(), *Class->GetName());
+			check(false);
+		}
+	}
+}
+
+//
+// Parses the header associated with the specified class.
+// Returns result enumeration.
+//
+ECompilationResult::Type FHeaderParser::ParseHeader(FClasses& AllClasses, FUnrealSourceFile* SourceFile)
+{
+	SetCurrentSourceFile(SourceFile);
+	FUnrealSourceFile* CurrentSrcFile = SourceFile;
+	if (CurrentSrcFile->IsParsed())
+	{
+		return ECompilationResult::Succeeded;
+	}
+
+	CurrentSrcFile->MarkAsParsed();
+
+	// Early-out if this class has previously failed some aspect of parsing
+	if (FailedFilesAnnotation.Get(CurrentSrcFile))
+	{
+		return ECompilationResult::OtherCompilationError;
+	}
+
+	// Reset the parser to begin a new class
+	bEncounteredNewStyleClass_UnmatchedBrackets = false;
+	bSpottedAutogeneratedHeaderInclude          = false;
+	bHaveSeenUClass                             = false;
+	bClassHasGeneratedBody                      = false;
+	bClassHasGeneratedUInterfaceBody            = false;
+	bClassHasGeneratedIInterfaceBody            = false;
+
+	ECompilationResult::Type Result = ECompilationResult::OtherCompilationError;
+
+	// Message.
+	UE_LOG(LogCompile, Verbose, TEXT("Parsing %s"), *CurrentSrcFile->GetFilename());
+
+	// Init compiler variables.
+	ResetParser(*CurrentSrcFile->GetContent());
+
+	// Init nesting.
+	NestLevel = 0;
+	TopNest = NULL;
+	PushNest(ENestType::GlobalScope, nullptr, CurrentSrcFile);
+
+	// C++ classes default to private access level
+	CurrentAccessSpecifier = ACCESS_Private; 
+
+	// Try to compile it, and catch any errors.
+	bool bEmptyFile = true;
+
+	// Tells if this header defines no-export classes only.
+	bool bNoExportClassesOnly = true;
+
+#if !PLATFORM_EXCEPTIONS_DISABLED
+	try
+#endif
+	{
+		// Parse entire program.
+		TArray<UDelegateFunction*> DelegatesToFixup;
+		while (CompileStatement(AllClasses, DelegatesToFixup))
+		{
+			bEmptyFile = false;
+
+			// Clear out the previous comment in anticipation of the next statement.
+			ClearComment();
+			StatementsParsed++;
+		}
+
+		PopNest(ENestType::GlobalScope, TEXT("Global scope"));
+
+		auto ScopeTypeIterator = CurrentSrcFile->GetScope()->GetTypeIterator();
+		while (ScopeTypeIterator.MoveNext())
+		{
+			UField* Type = *ScopeTypeIterator;
+
+			if (!Type->IsA<UScriptStruct>() && !Type->IsA<UClass>())
+			{
+				continue;
+			}
+
+			UStruct* Struct = Cast<UStruct>(Type);
+
+			// now validate all delegate variables declared in the class
+			TMap<FName, UFunction*> DelegateCache;
+			FixupDelegateProperties(AllClasses, Struct, FScope::GetTypeScope(Struct).Get(), DelegateCache);
+		}
+
+		// Fix up any delegates themselves, if they refer to other delegates
+		{
+			TMap<FName, UFunction*> DelegateCache;
+			for (UDelegateFunction* Delegate : DelegatesToFixup)
+			{
+				FixupDelegateProperties(AllClasses, Delegate, CurrentSrcFile->GetScope().Get(), DelegateCache);
+			}
+		}
+
+		// Precompute info for runtime optimization.
+		LinesParsed += InputLine;
+
+		if (RPCsNeedingHookup.Num() > 0)
+		{
+			FString ErrorMsg(TEXT("Request functions missing response pairs:\r\n"));
+			for (TMap<int32, FString>::TConstIterator It(RPCsNeedingHookup); It; ++It)
+			{
+				ErrorMsg += FString::Printf(TEXT("%s missing id %d\r\n"), *It.Value(), It.Key());
+			}
+
+			RPCsNeedingHookup.Empty();
+			FError::Throwf(*ErrorMsg);
+		}
+
+		// Make sure the compilation ended with valid nesting.
+		if (bEncounteredNewStyleClass_UnmatchedBrackets)
+		{
+			FError::Throwf(TEXT("Missing } at end of class") );
+		}
+
+		if (NestLevel == 1)
+		{
+			FError::Throwf(TEXT("Internal nest inconsistency") );
+		}
+		else if (NestLevel > 2)
+		{
+			FError::Throwf(TEXT("Unexpected end of script in '%s' block"), NestTypeName(TopNest->NestType) );
+		}
+
+		// First-pass success.
+		Result = ECompilationResult::Succeeded;
+
+		for (const TPair<UClass*, FSimplifiedParsingClassInfo>& ClassDataPair : CurrentSrcFile->GetDefinedClassesWithParsingInfo())
+		{
+			UClass* Class = ClassDataPair.Key;
+			PostParsingClassSetup(Class);
+
+			// Clean up and exit.
+			Class->Bind();
+
+			// Finalize functions
+			FinalizeScriptExposedFunctions(Class);
+
+			bNoExportClassesOnly = bNoExportClassesOnly && Class->HasAnyClassFlags(CLASS_NoExport);
+		}
+
+		check(CurrentSrcFile->IsParsed());
+
+		if (!bSpottedAutogeneratedHeaderInclude && !bEmptyFile && !bNoExportClassesOnly)
+		{
+			const FString& ExpectedHeaderName = CurrentSrcFile->GetGeneratedHeaderFilename();
+			FError::Throwf(TEXT("Expected an include at the top of the header: '#include \"%s\"'"), *ExpectedHeaderName);
+		}
+
+		GScriptHelper.CheckForNoIInterfaces();
+	}
+#if !PLATFORM_EXCEPTIONS_DISABLED
+	catch( TCHAR* ErrorMsg )
+	{
+		if (NestLevel == 0)
+		{
+			// Pushing nest so there is a file context for this error.
+			PushNest(ENestType::GlobalScope, nullptr, CurrentSrcFile);
+		}
+
+		// Handle compiler error.
+		{
+			TGuardValue<ELogTimes::Type> DisableLogTimes(GPrintLogTimes, ELogTimes::None);
+			FString FormattedErrorMessageWithContext = FString::Printf(TEXT("%s: Error: %s"), *GetContext(), ErrorMsg);
+
+			UE_LOG(LogCompile, Log,  TEXT("%s"), *FormattedErrorMessageWithContext );
+			Warn->Log(ELogVerbosity::Error, *FString::Printf(TEXT("Error: %s"), ErrorMsg));
+		}
+
+		FailedFilesAnnotation.Set(CurrentSrcFile);
+		Result = GCompilationResult;
+	}
+#endif
+
+	return Result; //@TODO: UCREMOVAL: This function is always returning succeeded even on a compiler error; should this continue?
+}
+
+/*-----------------------------------------------------------------------------
+	Global functions.
+-----------------------------------------------------------------------------*/
+
+ECompilationResult::Type FHeaderParser::ParseRestOfModulesSourceFiles(FClasses& AllClasses, UPackage* ModulePackage, FHeaderParser& HeaderParser)
+{
+	if (const TArray<FUnrealSourceFile*>* SourceFiles = GUnrealSourceFilesMap.FindFilesForPackage(ModulePackage))
+	{
+		for (FUnrealSourceFile* SourceFile : *SourceFiles)
+		{
+			if (!SourceFile->IsParsed() || SourceFile->GetDefinedClassesCount() == 0)
+			{
+				ECompilationResult::Type Result;
+				if ((Result = ParseHeaders(AllClasses, HeaderParser, SourceFile)) != ECompilationResult::Succeeded)
+				{
+					return Result;
+				}
+			}
+		}
+	}
+
+	return ECompilationResult::Succeeded;
+}
+
+// Parse Class's annotated headers and optionally its child classes.
+static const FString ObjectHeader(TEXT("NoExportTypes.h"));
+
+ECompilationResult::Type FHeaderParser::ParseHeaders(FClasses& AllClasses, FHeaderParser& HeaderParser, FUnrealSourceFile* SourceFile)
+{
+	ECompilationResult::Type Result = ECompilationResult::Succeeded;
+
+	if (SourceFile->AreDependenciesResolved())
+	{
+		return Result;
+	}
+
+	SourceFile->MarkDependenciesResolved();
+
+	TArray<FUnrealSourceFile*> SourceFilesRequired;
+
+	for (FHeaderProvider& Include : SourceFile->GetIncludes())
+	{
+		if (Include.GetId() == ObjectHeader)
+		{
+			continue;
+		}
+
+		if (FUnrealSourceFile* DepFile = Include.Resolve())
+		{
+			SourceFilesRequired.Add(DepFile);
+		}
+	}
+
+	for (const TPair<UClass*, FSimplifiedParsingClassInfo>& ClassDataPair : SourceFile->GetDefinedClassesWithParsingInfo())
+	{
+		UClass* Class = ClassDataPair.Key;
+		for (UClass* ParentClass = Class->GetSuperClass(); ParentClass && !ParentClass->HasAnyClassFlags(CLASS_Parsed | CLASS_Intrinsic); ParentClass = ParentClass->GetSuperClass())
+		{
+			SourceFilesRequired.Add(&GTypeDefinitionInfoMap[ParentClass]->GetUnrealSourceFile());
+		}
+	}
+
+	for (FUnrealSourceFile* RequiredFile : SourceFilesRequired)
+	{
+		SourceFile->GetScope()->IncludeScope(&RequiredFile->GetScope().Get());
+
+		ECompilationResult::Type ParseResult = ParseHeaders(AllClasses, HeaderParser, RequiredFile);
+
+		if (ParseResult != ECompilationResult::Succeeded)
+		{
+			return ParseResult;
+		}
+	}
+
+	// Parse the file
+	{
+		ECompilationResult::Type OneFileResult = HeaderParser.ParseHeader(AllClasses, SourceFile);
+
+		for (const TPair<UClass*, FSimplifiedParsingClassInfo>& ClassDataPair : SourceFile->GetDefinedClassesWithParsingInfo())
+		{
+			UClass* Class = ClassDataPair.Key;
+			Class->ClassFlags |= CLASS_Parsed;
+		}
+
+		if (OneFileResult != ECompilationResult::Succeeded)
+		{
+			// if we couldn't parse this file fail.
+			return OneFileResult;
+		}
+	}
+
+	// Success.
+	return Result;
+}
+
+bool FHeaderParser::DependentClassNameFromHeader(const TCHAR* HeaderFilename, FString& OutClassName)
+{
+	FString DependentClassName(HeaderFilename);
+	const int32 ExtensionIndex = DependentClassName.Find(TEXT("."), ESearchCase::CaseSensitive);
+	if (ExtensionIndex != INDEX_NONE)
+	{
+		// Generate UHeaderName name for this header.
+		OutClassName = TEXT("U") + FPaths::GetBaseFilename(MoveTemp(DependentClassName));
+		return true;
+	}
+	return false;
+}
+
+/**
+ * Gets source files ordered by UCLASSes inheritance.
+ *
+ * @param CurrentPackage Current package.
+ * @param AllClasses Current class tree.
+ *
+ * @returns Array of source files.
+ */
+TSet<FUnrealSourceFile*> GetSourceFilesWithInheritanceOrdering(UPackage* CurrentPackage, FClasses& AllClasses)
+{
+	TSet<FUnrealSourceFile*> SourceFiles;
+
+	TArray<FClass*> Classes = AllClasses.GetClassesInPackage();
+
+	// First add source files with the inheritance order.
+	for (UClass* Class : Classes)
+	{
+		TSharedRef<FUnrealTypeDefinitionInfo>* DefinitionInfoPtr = GTypeDefinitionInfoMap.Find(Class);
+		if (DefinitionInfoPtr == nullptr)
+		{
+			continue;
+		}
+
+		FUnrealSourceFile& SourceFile = (*DefinitionInfoPtr)->GetUnrealSourceFile();
+
+		if (SourceFile.GetScope()->ContainsTypes())
+		{
+			SourceFiles.Add(&SourceFile);
+		}
+	}
+
+	// Then add the rest.
+	if (const TArray<FUnrealSourceFile*>* SourceFilesForPackage = GUnrealSourceFilesMap.FindFilesForPackage(CurrentPackage))
+	{
+		for (FUnrealSourceFile* SourceFile : *SourceFilesForPackage)
+		{
+			if (SourceFile->GetScope()->ContainsTypes())
+			{
+				SourceFiles.Add(SourceFile);
+			}
+		}
+	}
+
+	return SourceFiles;
+}
+
+// Begins the process of exporting C++ class declarations for native classes in the specified package
+void FHeaderParser::ExportNativeHeaders(
+	UPackage* CurrentPackage,
+	FClasses& AllClasses,
+	bool bAllowSaveExportedHeaders,
+	const FManifestModule& Module
+)
+{
+	TSet<FUnrealSourceFile*> SourceFiles = GetSourceFilesWithInheritanceOrdering(CurrentPackage, AllClasses);
+	if (SourceFiles.Num() > 0)
+	{
+		if ( CurrentPackage != NULL )
+		{
+			UE_LOG(LogCompile, Verbose, TEXT("Exporting native class declarations for %s"), *CurrentPackage->GetName());
+		}
+		else
+		{
+			UE_LOG(LogCompile, Verbose, TEXT("Exporting native class declarations"));
+		}
+
+		// Export native class definitions to package header files.
+		FNativeClassHeaderGenerator(
+			CurrentPackage,
+			SourceFiles,
+			AllClasses,
+			bAllowSaveExportedHeaders
+		);
+	}
+}
+
+FHeaderParser::FHeaderParser(FFeedbackContext* InWarn, const FManifestModule& InModule)
+	: FBaseParser()
+	, Warn(InWarn)
+	, bSpottedAutogeneratedHeaderInclude(false)
+	, NestLevel(0)
+	, TopNest(nullptr)
+	, CurrentlyParsedModule(&InModule)
+{
+	// Determine if the current module is part of the engine or a game (we are more strict about things for Engine modules)
+	switch (InModule.ModuleType)
+	{
+	case EBuildModuleType::Program:
+		{
+			const FString AbsoluteEngineDir = FPaths::ConvertRelativePathToFull(FPaths::EngineDir());
+			const FString ModuleDir = FPaths::ConvertRelativePathToFull(InModule.BaseDirectory);
+			bIsCurrentModulePartOfEngine = ModuleDir.StartsWith(AbsoluteEngineDir);
+		}
+		break;
+	case EBuildModuleType::EngineRuntime:
+	case EBuildModuleType::EngineUncooked:
+	case EBuildModuleType::EngineDeveloper:
+	case EBuildModuleType::EngineEditor:
+	case EBuildModuleType::EngineThirdParty:
+		bIsCurrentModulePartOfEngine = true;
+		break;
+	case EBuildModuleType::GameRuntime:
+	case EBuildModuleType::GameUncooked:
+	case EBuildModuleType::GameDeveloper:
+	case EBuildModuleType::GameEditor:
+	case EBuildModuleType::GameThirdParty:
+		bIsCurrentModulePartOfEngine = false;
+		break;
+	default:
+		bIsCurrentModulePartOfEngine = true;
+		check(false);
+	}
+
+	FScriptLocation::Compiler = this;
+
+	static bool bConfigOptionsInitialized = false;
+
+	if (!bConfigOptionsInitialized)
+	{
+		// Read Ini options, GConfig must exist by this point
+		check(GConfig);
+
+		const FName TypeRedirectsKey(TEXT("TypeRedirects"));
+		const FName StructsWithNoPrefixKey(TEXT("StructsWithNoPrefix"));
+		const FName StructsWithTPrefixKey(TEXT("StructsWithTPrefix"));
+		const FName DelegateParameterCountStringsKey(TEXT("DelegateParameterCountStrings"));
+		const FName GeneratedCodeVersionKey(TEXT("GeneratedCodeVersion"));
+
+		FConfigSection* ConfigSection = GConfig->GetSectionPrivate(TEXT("UnrealHeaderTool"), false, true, GEngineIni);
+		if (ConfigSection)
+		{
+			for (FConfigSection::TIterator It(*ConfigSection); It; ++It)
+			{
+				if (It.Key() == TypeRedirectsKey)
+				{
+					FString OldType;
+					FString NewType;
+
+					FParse::Value(*It.Value().GetValue(), TEXT("OldType="), OldType);
+					FParse::Value(*It.Value().GetValue(), TEXT("NewType="), NewType);
+
+					TypeRedirectMap.Add(MoveTemp(OldType), MoveTemp(NewType));
+				}
+				else if (It.Key() == StructsWithNoPrefixKey)
+				{
+					StructsWithNoPrefix.Add(It.Value().GetValue());
+				}
+				else if (It.Key() == StructsWithTPrefixKey)
+				{
+					StructsWithTPrefix.Add(It.Value().GetValue());
+				}
+				else if (It.Key() == DelegateParameterCountStringsKey)
+				{
+					DelegateParameterCountStrings.Add(It.Value().GetValue());
+				}
+				else if (It.Key() == GeneratedCodeVersionKey)
+				{
+					DefaultGeneratedCodeVersion = ToGeneratedCodeVersion(It.Value().GetValue());
+				}
+			}
+		}
+		bConfigOptionsInitialized = true;
+	}
+}
+
+// Throws if a specifier value wasn't provided
+void FHeaderParser::RequireSpecifierValue(const FPropertySpecifier& Specifier, bool bRequireExactlyOne)
+{
+	if (Specifier.Values.Num() == 0)
+	{
+		FError::Throwf(TEXT("The specifier '%s' must be given a value"), *Specifier.Key);
+	}
+	else if ((Specifier.Values.Num() != 1) && bRequireExactlyOne)
+	{
+		FError::Throwf(TEXT("The specifier '%s' must be given exactly one value"), *Specifier.Key);
+	}
+}
+
+// Throws if a specifier value wasn't provided
+FString FHeaderParser::RequireExactlyOneSpecifierValue(const FPropertySpecifier& Specifier)
+{
+	RequireSpecifierValue(Specifier, /*bRequireExactlyOne*/ true);
+	return Specifier.Values[0];
+}
+
+// Exports the class to all vailable plugins
+void ExportClassToScriptPlugins(UClass* Class, const FManifestModule& Module, IScriptGeneratorPluginInterface& ScriptPlugin)
+{
+	TSharedRef<FUnrealTypeDefinitionInfo>* DefinitionInfoRef = GTypeDefinitionInfoMap.Find(Class);
+	if (DefinitionInfoRef == nullptr)
+	{
+		const FString Empty = TEXT("");
+		ScriptPlugin.ExportClass(Class, Empty, Empty, false);
+	}
+	else
+	{
+		FUnrealSourceFile& SourceFile = (*DefinitionInfoRef)->GetUnrealSourceFile();
+		ScriptPlugin.ExportClass(Class, SourceFile.GetFilename(), SourceFile.GetGeneratedFilename(), SourceFile.HasChanged());
+	}
+}
+
+// Exports class tree to all available plugins
+void ExportClassTreeToScriptPlugins(const FClassTree* Node, const FManifestModule& Module, IScriptGeneratorPluginInterface& ScriptPlugin)
+{
+	for (int32 ChildIndex = 0; ChildIndex < Node->NumChildren(); ++ChildIndex)
+	{
+		const FClassTree* ChildNode = Node->GetChild(ChildIndex);
+		ExportClassToScriptPlugins(ChildNode->GetClass(), Module, ScriptPlugin);
+	}
+
+	for (int32 ChildIndex = 0; ChildIndex < Node->NumChildren(); ++ChildIndex)
+	{
+		const FClassTree* ChildNode = Node->GetChild(ChildIndex);
+		ExportClassTreeToScriptPlugins(ChildNode, Module, ScriptPlugin);
+	}
+}
+
+// Parse all headers for classes that are inside CurrentPackage.
+ECompilationResult::Type FHeaderParser::ParseAllHeadersInside(
+	FClasses& ModuleClasses,
+	FFeedbackContext* Warn,
+	UPackage* CurrentPackage,
+	const FManifestModule& Module,
+	TArray<IScriptGeneratorPluginInterface*>& ScriptPlugins
+	)
+{
+	SCOPE_SECONDS_COUNTER_UHT(ParseAllHeaders);
+
+	// Disable loading of objects outside of this package (or more exactly, objects which aren't UFields, CDO, or templates)
+	TGuardValue<bool> AutoRestoreVerifyObjectRefsFlag(GVerifyObjectReferencesOnly, true);
+	// Create the header parser and register it as the warning context.
+	// Note: This must be declared outside the try block, since the catch block will log into it.
+	FHeaderParser HeaderParser(Warn, Module);
+	Warn->SetContext(&HeaderParser);
+
+
+	// Hierarchically parse all classes.
+	ECompilationResult::Type Result = ECompilationResult::Succeeded;
+#if !PLATFORM_EXCEPTIONS_DISABLED
+	try
+#endif
+	{
+		FName ModuleName = FName(*Module.Name);
+		bool bNeedsRegeneration = Module.NeedsRegeneration();
+
+		// Set up a filename for the error context if we don't even get as far parsing a class
+		FClass*                                      RootClass          = ModuleClasses.GetRootClass();
+		const TSharedRef<FUnrealTypeDefinitionInfo>& TypeDefinitionInfo = GTypeDefinitionInfoMap[RootClass];
+		const FUnrealSourceFile&                     RootSourceFile     = TypeDefinitionInfo->GetUnrealSourceFile();
+		const FString&                               RootFilename       = RootSourceFile.GetFilename();
+
+		HeaderParser.Filename = IFileManager::Get().ConvertToAbsolutePathForExternalAppForRead(*RootFilename);
+
+		if (const TArray<FUnrealSourceFile*>* SourceFiles = GPublicSourceFileSet.FindFilesForPackage(CurrentPackage))
+		{
+			for (FUnrealSourceFile* SourceFile : *SourceFiles)
+			{
+				if ((!SourceFile->IsParsed() || SourceFile->GetDefinedClassesCount() == 0))
+				{
+					Result = ParseHeaders(ModuleClasses, HeaderParser, SourceFile);
+					if (Result != ECompilationResult::Succeeded)
+					{
+						return Result;
+					}
+				}
+			}
+		}
+		if (Result == ECompilationResult::Succeeded)
+		{
+			Result = FHeaderParser::ParseRestOfModulesSourceFiles(ModuleClasses, CurrentPackage, HeaderParser);
+		}
+
+		if (Result == ECompilationResult::Succeeded)
+		{
+			// Validate the sparse class data for all classes in the current package
+			for (const FClass* Class : ModuleClasses.GetClassesInPackage(CurrentPackage))
+			{
+				CheckSparseClassData(Class);
+			}
+
+			// Export the autogenerated code wrappers
+
+			// At this point all headers have been parsed and the header parser will
+			// no longer have up to date info about what's being done so unregister it 
+			// from the feedback context.
+			Warn->SetContext(NULL);
+
+			double ExportTime = 0.0;
+			{
+				FScopedDurationTimer Timer(ExportTime);
+				ExportNativeHeaders(
+					CurrentPackage,
+					ModuleClasses,
+					Module.SaveExportedHeaders,
+					Module
+				);
+			}
+			GHeaderCodeGenTime += ExportTime;
+
+			// Done with header generation
+			if (HeaderParser.LinesParsed > 0)
+			{
+				UE_LOG(LogCompile, Log, TEXT("Success: Parsed %i line(s), %i statement(s) in %.2f secs.\r\n"), HeaderParser.LinesParsed, HeaderParser.StatementsParsed, ExportTime);
+			}
+			else
+			{
+				UE_LOG(LogCompile, Log, TEXT("Success: Everything is up to date (in %.2f secs)"), ExportTime);
+			}
+		}
+	}
+#if !PLATFORM_EXCEPTIONS_DISABLED
+	catch (TCHAR* ErrorMsg)
+	{
+		Warn->Log(ELogVerbosity::Error, ErrorMsg);
+		Result = GCompilationResult;
+	}
+#endif
+	// Unregister the header parser from the feedback context
+	Warn->SetContext(NULL);
+
+	if (Result == ECompilationResult::Succeeded && ScriptPlugins.Num())
+	{
+		FScopedDurationTimer PluginTimeTracker(GPluginOverheadTime);
+
+		FClassTree* RootNode = &ModuleClasses.GetClassTree();
+		for (IScriptGeneratorPluginInterface* Plugin : ScriptPlugins)
+		{
+			if (Plugin->ShouldExportClassesForModule(Module.Name, Module.ModuleType, Module.GeneratedIncludeDirectory))
+			{
+				ExportClassToScriptPlugins(RootNode->GetClass(), Module, *Plugin);
+				ExportClassTreeToScriptPlugins(RootNode, Module, *Plugin);
+			}
+		}
+	}
+
+	return Result;
+}
+
+/** 
+ * Returns True if the given class name includes a valid Unreal prefix and matches up with the given original class.
+ *
+ * @param InNameToCheck - Name w/ potential prefix to check
+ * @param OriginalClassName - Name of class w/ no prefix to check against
+ */
+bool FHeaderParser::ClassNameHasValidPrefix(const FString& InNameToCheck, const FString& OriginalClassName)
+{
+	bool bIsLabledDeprecated;
+	const FString ClassPrefix = GetClassPrefix( InNameToCheck, bIsLabledDeprecated );
+
+	// If the class is labeled deprecated, don't try to resolve it during header generation, valid results can't be guaranteed.
+	if (bIsLabledDeprecated)
+	{
+		return true;
+	}
+
+	if (ClassPrefix.IsEmpty())
+	{
+		return false;
+	}
+
+	FString TestString = FString::Printf(TEXT("%s%s"), *ClassPrefix, *OriginalClassName);
+
+	const bool bNamesMatch = ( InNameToCheck == *TestString );
+
+	return bNamesMatch;
+}
+
+void FHeaderParser::ParseClassName(const TCHAR* Temp, FString& ClassName)
+{
+	// Skip leading whitespace
+	while (FChar::IsWhitespace(*Temp))
+	{
+		++Temp;
+	}
+
+	// Run thru characters (note: relying on later code to reject the name for a leading number, etc...)
+	const TCHAR* StringStart = Temp;
+	while (FChar::IsAlnum(*Temp) || FChar::IsUnderscore(*Temp))
+	{
+		++Temp;
+	}
+
+	ClassName = FString(UE_PTRDIFF_TO_INT32(Temp - StringStart), StringStart);
+	if (ClassName.EndsWith(TEXT("_API"), ESearchCase::CaseSensitive))
+	{
+		// RequiresAPI token for a given module
+
+		//@TODO: UCREMOVAL: Validate the module name
+		FString RequiresAPISymbol = ClassName;
+
+		// Now get the real class name
+		ClassName.Empty();
+		ParseClassName(Temp, ClassName);
+	}
+}
+
+enum class EBlockDirectiveType
+{
+	// We're in a CPP block
+	CPPBlock,
+
+	// We're in a !CPP block
+	NotCPPBlock,
+
+	// We're in a 0 block
+	ZeroBlock,
+
+	// We're in a 1 block
+	OneBlock,
+
+	// We're in a WITH_HOT_RELOAD block
+	WithHotReload,
+
+	// We're in a WITH_EDITOR block
+	WithEditor,
+
+	// We're in a WITH_EDITORONLY_DATA block
+	WithEditorOnlyData,
+
+	// We're in a block with an unrecognized directive
+	UnrecognizedBlock
+};
+
+bool ShouldKeepBlockContents(EBlockDirectiveType DirectiveType)
+{
+	switch (DirectiveType)
+	{
+		case EBlockDirectiveType::NotCPPBlock:
+		case EBlockDirectiveType::OneBlock:
+		case EBlockDirectiveType::WithHotReload:
+		case EBlockDirectiveType::WithEditor:
+		case EBlockDirectiveType::WithEditorOnlyData:
+			return true;
+
+		case EBlockDirectiveType::CPPBlock:
+		case EBlockDirectiveType::ZeroBlock:
+		case EBlockDirectiveType::UnrecognizedBlock:
+			return false;
+	}
+
+	check(false);
+	UE_ASSUME(false);
+}
+
+bool ShouldKeepDirective(EBlockDirectiveType DirectiveType)
+{
+	switch (DirectiveType)
+	{
+		case EBlockDirectiveType::WithHotReload:
+		case EBlockDirectiveType::WithEditor:
+		case EBlockDirectiveType::WithEditorOnlyData:
+			return true;
+
+		case EBlockDirectiveType::CPPBlock:
+		case EBlockDirectiveType::NotCPPBlock:
+		case EBlockDirectiveType::ZeroBlock:
+		case EBlockDirectiveType::OneBlock:
+		case EBlockDirectiveType::UnrecognizedBlock:
+			return false;
+	}
+
+	check(false);
+	UE_ASSUME(false);
+}
+
+EBlockDirectiveType ParseCommandToBlockDirectiveType(const TCHAR** Str)
+{
+	if (FParse::Command(Str, TEXT("0")))
+	{
+		return EBlockDirectiveType::ZeroBlock;
+	}
+
+	if (FParse::Command(Str, TEXT("1")))
+	{
+		return EBlockDirectiveType::OneBlock;
+	}
+
+	if (FParse::Command(Str, TEXT("CPP")))
+	{
+		return EBlockDirectiveType::CPPBlock;
+	}
+
+	if (FParse::Command(Str, TEXT("!CPP")))
+	{
+		return EBlockDirectiveType::NotCPPBlock;
+	}
+
+	if (FParse::Command(Str, TEXT("WITH_HOT_RELOAD")))
+	{
+		return EBlockDirectiveType::WithHotReload;
+	}
+
+	if (FParse::Command(Str, TEXT("WITH_EDITOR")))
+	{
+		return EBlockDirectiveType::WithEditor;
+	}
+
+	if (FParse::Command(Str, TEXT("WITH_EDITORONLY_DATA")))
+	{
+		return EBlockDirectiveType::WithEditorOnlyData;
+	}
+
+	return EBlockDirectiveType::UnrecognizedBlock;
+}
+
+const TCHAR* GetBlockDirectiveTypeString(EBlockDirectiveType DirectiveType)
+{
+	switch (DirectiveType)
+	{
+		case EBlockDirectiveType::CPPBlock:           return TEXT("CPP");
+		case EBlockDirectiveType::NotCPPBlock:        return TEXT("!CPP");
+		case EBlockDirectiveType::ZeroBlock:          return TEXT("0");
+		case EBlockDirectiveType::OneBlock:           return TEXT("1");
+		case EBlockDirectiveType::WithHotReload:      return TEXT("WITH_HOT_RELOAD");
+		case EBlockDirectiveType::WithEditor:         return TEXT("WITH_EDITOR");
+		case EBlockDirectiveType::WithEditorOnlyData: return TEXT("WITH_EDITORONLY_DATA");
+		case EBlockDirectiveType::UnrecognizedBlock:  return TEXT("<unrecognized>");
+	}
+
+	check(false);
+	UE_ASSUME(false);
+}
+
+// Performs a preliminary parse of the text in the specified buffer, pulling out useful information for the header generation process
+void FHeaderParser::SimplifiedClassParse(const TCHAR* Filename, const TCHAR* InBuffer, TArray<FSimplifiedParsingClassInfo>& OutParsedClassArray, TArray<FHeaderProvider>& DependentOn, FStringOutputDevice& ClassHeaderTextStrippedOfCppText)
+{
+	FHeaderPreParser Parser;
+	FString StrLine;
+	FString ClassName;
+	FString BaseClassName;
+
+	// Two passes, preprocessor, then looking for the class stuff
+
+	// The layer of multi-line comment we are in.
+	int32 CurrentLine = 0;
+	const TCHAR* Buffer = InBuffer;
+
+	// Preprocessor pass
+	while (FParse::Line(&Buffer, StrLine, true))
+	{
+		CurrentLine++;
+		const TCHAR* Str = *StrLine;
+		int32 BraceCount = 0;
+
+		bool bIf = FParse::Command(&Str,TEXT("#if"));
+		if( bIf || FParse::Command(&Str,TEXT("#ifdef")) || FParse::Command(&Str,TEXT("#ifndef")) )
+		{
+			EBlockDirectiveType RootDirective;
+			if (bIf)
+			{
+				RootDirective = ParseCommandToBlockDirectiveType(&Str);
+			}
+			else
+			{
+				// #ifdef or #ifndef are always treated as CPP
+				RootDirective = EBlockDirectiveType::UnrecognizedBlock;
+			}
+
+			TArray<EBlockDirectiveType, TInlineAllocator<8>> DirectiveStack;
+			DirectiveStack.Push(RootDirective);
+
+			bool bShouldKeepBlockContents = ShouldKeepBlockContents(RootDirective);
+			bool bIsZeroBlock = RootDirective == EBlockDirectiveType::ZeroBlock;
+
+			ClassHeaderTextStrippedOfCppText.Logf(TEXT("%s\r\n"), ShouldKeepDirective(RootDirective) ? *StrLine : TEXT(""));
+
+			while ((DirectiveStack.Num() > 0) && FParse::Line(&Buffer, StrLine, 1))
+			{
+				CurrentLine++;
+				Str = *StrLine;
+
+				bool bShouldKeepLine = bShouldKeepBlockContents;
+
+				bool bIsDirective = false;
+				if( FParse::Command(&Str,TEXT("#endif")) )
+				{
+					EBlockDirectiveType OldDirective = DirectiveStack.Pop();
+
+					bShouldKeepLine &= ShouldKeepDirective(OldDirective);
+					bIsDirective     = true;
+				}
+				else if( FParse::Command(&Str,TEXT("#if")) || FParse::Command(&Str,TEXT("#ifdef")) || FParse::Command(&Str,TEXT("#ifndef")) )
+				{
+					EBlockDirectiveType Directive = ParseCommandToBlockDirectiveType(&Str);
+					DirectiveStack.Push(Directive);
+
+					bShouldKeepLine &= ShouldKeepDirective(Directive);
+					bIsDirective     = true;
+				}
+				else if (FParse::Command(&Str,TEXT("#elif")))
+				{
+					EBlockDirectiveType NewDirective = ParseCommandToBlockDirectiveType(&Str);
+					EBlockDirectiveType OldDirective = DirectiveStack.Top();
+
+					// Check to see if we're mixing ignorable directive types - we don't support this
+					bool bKeepNewDirective = ShouldKeepDirective(NewDirective);
+					bool bKeepOldDirective = ShouldKeepDirective(OldDirective);
+					if (bKeepNewDirective != bKeepOldDirective)
+					{
+						FFileLineException::Throwf(
+							Filename,
+							CurrentLine,
+							TEXT("Mixing %s with %s in an #elif preprocessor block is not supported"),
+							GetBlockDirectiveTypeString(OldDirective),
+							GetBlockDirectiveTypeString(NewDirective)
+						);
+					}
+
+					DirectiveStack.Top() = NewDirective;
+
+					bShouldKeepLine &= bKeepNewDirective;
+					bIsDirective     = true;
+				}
+				else if (FParse::Command(&Str, TEXT("#else")))
+				{
+					switch (DirectiveStack.Top())
+					{
+						case EBlockDirectiveType::ZeroBlock:
+							DirectiveStack.Top() = EBlockDirectiveType::OneBlock;
+							break;
+
+						case EBlockDirectiveType::OneBlock:
+							DirectiveStack.Top() = EBlockDirectiveType::ZeroBlock;
+							break;
+
+						case EBlockDirectiveType::CPPBlock:
+							DirectiveStack.Top() = EBlockDirectiveType::NotCPPBlock;
+							break;
+
+						case EBlockDirectiveType::NotCPPBlock:
+							DirectiveStack.Top() = EBlockDirectiveType::CPPBlock;
+							break;
+
+						case EBlockDirectiveType::WithHotReload:
+							FFileLineException::Throwf(Filename, CurrentLine, TEXT("Bad preprocessor directive in metadata declaration: %s; Only 'CPP', '1' and '0' can have #else directives"), *ClassName);
+
+						case EBlockDirectiveType::UnrecognizedBlock:
+						case EBlockDirectiveType::WithEditor:
+						case EBlockDirectiveType::WithEditorOnlyData:
+							// We allow unrecognized directives, WITH_EDITOR and WITH_EDITORONLY_DATA to have #else blocks.
+							// However, we don't actually change how UHT processes these #else blocks.
+							break;
+					}
+
+					bShouldKeepLine &= ShouldKeepDirective(DirectiveStack.Top());
+					bIsDirective     = true;
+				}
+				else
+				{
+					// Check for UHT identifiers inside skipped blocks, unless it's a zero block, because the compiler is going to skip those anyway.
+					if (!bShouldKeepBlockContents && !bIsZeroBlock)
+					{
+						auto FindInitialStr = [](const TCHAR*& FoundSubstr, const FString& StrToSearch, const TCHAR* ConstructName) -> bool
+						{
+							if (StrToSearch.StartsWith(ConstructName, ESearchCase::CaseSensitive))
+							{
+								FoundSubstr = ConstructName;
+								return true;
+							}
+
+							return false;
+						};
+
+						FString TrimmedStrLine = StrLine;
+						TrimmedStrLine.TrimStartInline();
+
+						const TCHAR* FoundSubstr = nullptr;
+						if (FindInitialStr(FoundSubstr, TrimmedStrLine, TEXT("UPROPERTY"))
+							|| FindInitialStr(FoundSubstr, TrimmedStrLine, TEXT("UCLASS"))
+							|| FindInitialStr(FoundSubstr, TrimmedStrLine, TEXT("USTRUCT"))
+							|| FindInitialStr(FoundSubstr, TrimmedStrLine, TEXT("UENUM"))
+							|| FindInitialStr(FoundSubstr, TrimmedStrLine, TEXT("UINTERFACE"))
+							|| FindInitialStr(FoundSubstr, TrimmedStrLine, TEXT("UDELEGATE"))
+							|| FindInitialStr(FoundSubstr, TrimmedStrLine, TEXT("UFUNCTION")))
+						{
+							FFileLineException::Throwf(Filename, CurrentLine, TEXT("%s must not be inside preprocessor blocks, except for WITH_EDITORONLY_DATA"), FoundSubstr);
+						}
+
+						// Try and determine if this line contains something like a serialize function
+						if (TrimmedStrLine.Len() > 0)
+						{
+							static const FString Str_Void = TEXT("void");
+							static const FString Str_Serialize = TEXT("Serialize(");
+							static const FString Str_FArchive = TEXT("FArchive");
+							static const FString Str_FStructuredArchive = TEXT("FStructuredArchive::FSlot");
+
+							int32 Pos = 0;
+							if ((Pos = TrimmedStrLine.Find(Str_Void, ESearchCase::CaseSensitive, ESearchDir::FromStart, Pos)) != -1)
+							{
+								Pos += Str_Void.Len();
+								if ((Pos = TrimmedStrLine.Find(Str_Serialize, ESearchCase::CaseSensitive, ESearchDir::FromStart, Pos)) != -1)
+								{
+									Pos += Str_Serialize.Len();
+
+									if (((TrimmedStrLine.Find(Str_FArchive, ESearchCase::CaseSensitive, ESearchDir::FromStart, Pos)) != -1) ||
+										((TrimmedStrLine.Find(Str_FStructuredArchive, ESearchCase::CaseSensitive, ESearchDir::FromStart, Pos)) != -1))
+									{
+										FFileLineException::Throwf(Filename, CurrentLine, TEXT("'%s' must not be inside preprocessor blocks, except for WITH_EDITORONLY_DATA"), *TrimmedStrLine);
+									}
+								}
+							}
+						}
+					}
+				}
+
+				ClassHeaderTextStrippedOfCppText.Logf(TEXT("%s\r\n"), bShouldKeepLine ? *StrLine : TEXT(""));
+
+				if (bIsDirective)
+				{
+					bShouldKeepBlockContents = Algo::AllOf(DirectiveStack, &ShouldKeepBlockContents);
+					bIsZeroBlock = DirectiveStack.Contains(EBlockDirectiveType::ZeroBlock);
+				}
+			}
+		}
+		else if ( FParse::Command(&Str,TEXT("#include")) )
+		{
+			ClassHeaderTextStrippedOfCppText.Logf( TEXT("%s\r\n"), *StrLine );
+		}
+		else
+		{
+			ClassHeaderTextStrippedOfCppText.Logf( TEXT("%s\r\n"), *StrLine );
+		}
+	}
+
+	// now start over go look for the class
+
+	int32 CommentDim  = 0;
+	CurrentLine = 0;
+	Buffer      = *ClassHeaderTextStrippedOfCppText;
+
+	const TCHAR* StartOfLine            = Buffer;
+	bool         bFoundGeneratedInclude = false;
+	bool         bFoundExportedClasses  = false;
+
+	while (FParse::Line(&Buffer, StrLine, true))
+	{
+		CurrentLine++;
+
+		const TCHAR* Str = *StrLine;
+		bool bProcess = CommentDim <= 0;	// for skipping nested multi-line comments
+
+		int32 BraceCount = 0;
+		if( bProcess && FParse::Command(&Str,TEXT("#if")) )
+		{
+		}
+		else if ( bProcess && FParse::Command(&Str,TEXT("#include")) )
+		{
+			// Handle #include directives as if they were 'dependson' keywords.
+			const FString& DependsOnHeaderName = Str;
+
+			if (DependsOnHeaderName != TEXT("\"UObject/DefineUPropertyMacros.h\"") && DependsOnHeaderName != TEXT("\"UObject/UndefineUPropertyMacros.h\""))
+			{
+				if (bFoundGeneratedInclude)
+				{
+					FFileLineException::Throwf(Filename, CurrentLine, TEXT("#include found after .generated.h file - the .generated.h file should always be the last #include in a header"));
+				}
+
+				bFoundGeneratedInclude = DependsOnHeaderName.Contains(TEXT(".generated.h"));
+				if (!bFoundGeneratedInclude && DependsOnHeaderName.Len())
+				{
+					bool  bIsQuotedInclude = DependsOnHeaderName[0] == '\"';
+					int32 HeaderFilenameEnd = DependsOnHeaderName.Find(bIsQuotedInclude ? TEXT("\"") : TEXT(">"), ESearchCase::CaseSensitive, ESearchDir::FromStart, 1);
+
+					if (HeaderFilenameEnd != INDEX_NONE)
+					{
+						// Include the extension in the name so that we later know where this entry came from.
+						DependentOn.Add(FHeaderProvider(EHeaderProviderSourceType::FileName, FPaths::GetCleanFilename(DependsOnHeaderName.Mid(1, HeaderFilenameEnd - 1))));
+					}
+				}
+			}
+		}
+		else if ( bProcess && FParse::Command(&Str,TEXT("#else")) )
+		{
+		}
+		else if ( bProcess && FParse::Command(&Str,TEXT("#elif")) )
+		{
+		}
+		else if ( bProcess && FParse::Command(&Str,TEXT("#endif")) )
+		{
+		}
+		else
+		{
+			int32 Pos = INDEX_NONE;
+			int32 EndPos = INDEX_NONE;
+			int32 StrBegin = INDEX_NONE;
+			int32 StrEnd = INDEX_NONE;
+				
+			bool bEscaped = false;
+			for ( int32 CharPos = 0; CharPos < StrLine.Len(); CharPos++ )
+			{
+				if ( bEscaped )
+				{
+					bEscaped = false;
+				}
+				else if ( StrLine[CharPos] == TEXT('\\') )
+				{
+					bEscaped = true;
+				}
+				else if ( StrLine[CharPos] == TEXT('\"') )
+				{
+					if ( StrBegin == INDEX_NONE )
+					{
+						StrBegin = CharPos;
+					}
+					else
+					{
+						StrEnd = CharPos;
+						break;
+					}
+				}
+			}
+
+			// Find the first '/' and check for '//' or '/*' or '*/'
+			if (StrLine.FindChar('/', Pos))
+			{
+				if (Pos >= 0)
+				{
+					// Stub out the comments, ignoring anything inside literal strings.
+					Pos = StrLine.Find(TEXT("//"), ESearchCase::CaseSensitive, ESearchDir::FromStart, Pos);
+
+					// Check if first slash is end of multiline comment and adjust position if necessary.
+					if (Pos > 0 && StrLine[Pos - 1] == TEXT('*'))
+					{
+						++Pos;
+					}
+
+					if (Pos >= 0)
+					{
+						if (StrBegin == INDEX_NONE || Pos < StrBegin || Pos > StrEnd)
+						{
+							StrLine.LeftInline(Pos, false);
+						}
+
+						if (StrLine.IsEmpty())
+						{
+							continue;
+						}
+					}
+
+					// look for a / * ... * / block, ignoring anything inside literal strings
+					Pos = StrLine.Find(TEXT("/*"), ESearchCase::CaseSensitive, ESearchDir::FromStart, Pos);
+					EndPos = StrLine.Find(TEXT("*/"), ESearchCase::CaseSensitive, ESearchDir::FromStart, FMath::Max(0, Pos - 1));
+					if (Pos >= 0)
+					{
+						if (StrBegin == INDEX_NONE || Pos < StrBegin || Pos > StrEnd)
+						{
+							if (EndPos != INDEX_NONE && (EndPos < StrBegin || EndPos > StrEnd))
+							{
+								StrLine = StrLine.Left(Pos) + StrLine.Mid(EndPos + 2);
+								EndPos = INDEX_NONE;
+							}
+							else
+							{
+								StrLine.LeftInline(Pos, false);
+								CommentDim++;
+							}
+						}
+						bProcess = CommentDim <= 1;
+					}
+
+					if (EndPos >= 0)
+					{
+						if (StrBegin == INDEX_NONE || EndPos < StrBegin || EndPos > StrEnd)
+						{
+							StrLine.MidInline(EndPos + 2, MAX_int32, false);
+							CommentDim--;
+						}
+
+						bProcess = CommentDim <= 0;
+					}
+				}
+			}
+
+			StrLine.TrimStartInline();
+			if (!bProcess || StrLine.IsEmpty())
+			{
+				continue;
+			}
+
+			Str = *StrLine;
+
+			// Get class or interface name
+			if (const TCHAR* UInterfaceMacroDecl = FCString::Strfind(Str, TEXT("UINTERFACE")))
+			{
+				if (UInterfaceMacroDecl == FCString::Strspn(Str, TEXT("\t ")) + Str)
+				{
+					if (UInterfaceMacroDecl[10] != TEXT('('))
+					{
+						FFileLineException::Throwf(Filename, CurrentLine, TEXT("Missing open parenthesis after UINTERFACE"));
+					}
+
+					FName StrippedInterfaceName;
+					Parser.ParseClassDeclaration(Filename, StartOfLine + (UInterfaceMacroDecl - Str), CurrentLine, TEXT("UINTERFACE"), /*out*/ StrippedInterfaceName, /*out*/ ClassName, /*out*/ BaseClassName, /*out*/ DependentOn, OutParsedClassArray);
+					OutParsedClassArray.Add(FSimplifiedParsingClassInfo(MoveTemp(ClassName), MoveTemp(BaseClassName), CurrentLine, true));
+					if (!bFoundExportedClasses)
+					{
+						if (FClassDeclarationMetaData* Found = GClassDeclarations.Find(StrippedInterfaceName))
+						{
+							bFoundExportedClasses = !(Found->ClassFlags & CLASS_NoExport);
+						}
+					}
+				}
+			}
+
+			if (const TCHAR* UClassMacroDecl = FCString::Strfind(Str, TEXT("UCLASS")))
+			{
+				if (UClassMacroDecl == FCString::Strspn(Str, TEXT("\t ")) + Str)
+				{
+					if (UClassMacroDecl[6] != TEXT('('))
+					{
+						FFileLineException::Throwf(Filename, CurrentLine, TEXT("Missing open parenthesis after UCLASS"));
+					}
+
+					FName StrippedClassName;
+					Parser.ParseClassDeclaration(Filename, StartOfLine + (UClassMacroDecl - Str), CurrentLine, TEXT("UCLASS"), /*out*/ StrippedClassName, /*out*/ ClassName, /*out*/ BaseClassName, /*out*/ DependentOn, OutParsedClassArray);
+					OutParsedClassArray.Add(FSimplifiedParsingClassInfo(MoveTemp(ClassName), MoveTemp(BaseClassName), CurrentLine, false));
+					if (!bFoundExportedClasses)
+					{
+						if (FClassDeclarationMetaData* Found = GClassDeclarations.Find(StrippedClassName))
+						{
+							bFoundExportedClasses = !(Found->ClassFlags & CLASS_NoExport);
+						}
+					}
+				}
+			}
+		}
+	
+		StartOfLine = Buffer;
+	}
+
+	if (bFoundExportedClasses && !bFoundGeneratedInclude)
+	{
+		FError::Throwf(TEXT("No #include found for the .generated.h file - the .generated.h file should always be the last #include in a header"));
+	}
+}
+
+/////////////////////////////////////////////////////
+// FHeaderPreParser
+
+void FHeaderPreParser::ParseClassDeclaration(const TCHAR* Filename, const TCHAR* InputText, int32 InLineNumber, const TCHAR* StartingMatchID, FName& out_StrippedClassName, FString& out_ClassName, FString& out_BaseClassName, TArray<FHeaderProvider>& out_RequiredIncludes, const TArray<FSimplifiedParsingClassInfo>& ParsedClassArray)
+{
+	const TCHAR* ErrorMsg = TEXT("Class declaration");
+
+	ResetParser(InputText, InLineNumber);
+
+	// Require 'UCLASS' or 'UINTERFACE'
+	RequireIdentifier(StartingMatchID, ESearchCase::CaseSensitive, ErrorMsg);
+
+	// New-style UCLASS() syntax
+	TMap<FName, FString> MetaData;
+	TArray<FPropertySpecifier> SpecifiersFound;
+	ReadSpecifierSetInsideMacro(SpecifiersFound, ErrorMsg, MetaData);
+
+	// Require 'class'
+	RequireIdentifier(TEXT("class"), ESearchCase::CaseSensitive, ErrorMsg);
+
+	// alignas() can come before or after the deprecation macro.
+	// We can't have both, but the compiler will catch that anyway.
+	SkipAlignasIfNecessary(*this);
+	SkipDeprecatedMacroIfNecessary(*this);
+	SkipAlignasIfNecessary(*this);
+
+	// Read the class name
+	FString RequiredAPIMacroIfPresent;
+	ParseNameWithPotentialAPIMacroPrefix(/*out*/ out_ClassName, /*out*/ RequiredAPIMacroIfPresent, StartingMatchID);
+
+	FString ClassNameWithoutPrefixStr = GetClassNameWithPrefixRemoved(out_ClassName);
+	out_StrippedClassName = *ClassNameWithoutPrefixStr;
+
+	{
+		auto ConstructDeclarationData = [MetaData = MoveTemp(MetaData), SpecifiersFound = MoveTemp(SpecifiersFound), &RequiredAPIMacroIfPresent]() mutable
+		{
+			// Add class declaration meta data so that we can access class flags before the class is fully parsed
+			TSharedRef<FClassDeclarationMetaData> DeclarationData = MakeShareable(new FClassDeclarationMetaData());
+			DeclarationData->MetaData = MoveTemp(MetaData);
+			DeclarationData->ParseClassProperties(MoveTemp(SpecifiersFound), RequiredAPIMacroIfPresent);
+			return DeclarationData;
+		};
+
+		GClassDeclarations.AddIfMissing(out_StrippedClassName, MoveTemp(ConstructDeclarationData));
+	}
+
+	// Skip optional final keyword
+	MatchIdentifier(TEXT("final"), ESearchCase::CaseSensitive);
+
+	// Handle inheritance
+	if (MatchSymbol(TEXT(':')))
+	{
+		// Require 'public'
+		RequireIdentifier(TEXT("public"), ESearchCase::CaseSensitive, ErrorMsg);
+
+		// Inherits from something
+		FToken BaseClassNameToken;
+		if (!GetIdentifier(BaseClassNameToken, true))
+		{
+			FError::Throwf(TEXT("Expected a base class name"));
+		}
+
+		out_BaseClassName = BaseClassNameToken.Identifier;
+
+		int32 InputLineLocal = InputLine;
+		auto AddDependencyIfNeeded = [Filename, InputLineLocal, &ParsedClassArray, &out_RequiredIncludes, &ClassNameWithoutPrefixStr](const FString& DependencyClassName)
+		{
+			if (!Algo::FindBy(ParsedClassArray, DependencyClassName, &FSimplifiedParsingClassInfo::GetClassName))
+			{
+				FString DependencyClassNameWithoutPrefixStr = GetClassNameWithPrefixRemoved(DependencyClassName);
+
+				if (ClassNameWithoutPrefixStr == DependencyClassNameWithoutPrefixStr)
+				{
+					FFileLineException::Throwf(Filename, InputLineLocal, TEXT("A class cannot inherit itself or a type with the same name but a different prefix"));
+				}
+
+				FString StrippedDependencyName = DependencyClassName.Mid(1);
+
+				// Only add a stripped dependency if the stripped name differs from the stripped class name
+				// otherwise it's probably a class with a different prefix.
+				if (StrippedDependencyName != ClassNameWithoutPrefixStr)
+				{
+					out_RequiredIncludes.Add(FHeaderProvider(EHeaderProviderSourceType::ClassName, MoveTemp(StrippedDependencyName)));
+				}
+			}
+		};
+
+		AddDependencyIfNeeded(out_BaseClassName);
+
+		// Get additional inheritance links and rack them up as dependencies if they're UObject derived
+		while (MatchSymbol(TEXT(',')))
+		{
+			// Require 'public'
+			RequireIdentifier(TEXT("public"), ESearchCase::CaseSensitive, ErrorMsg);
+
+			FToken InterfaceClassNameToken;
+			if (!GetIdentifier(InterfaceClassNameToken, true))
+			{
+				FFileLineException::Throwf(Filename, InputLine, TEXT("Expected an interface class name"));
+			}
+
+			AddDependencyIfNeeded(FString(InterfaceClassNameToken.Identifier));
+		}
+	}
+}
+
+//////////////////////////////////////////////////////////////////////////////////////////////////////////////
+
+bool FHeaderParser::DefaultValueStringCppFormatToInnerFormat(const FProperty* Property, const FString& CppForm, FString &OutForm)
+{
+	OutForm = FString();
+	if (!Property || CppForm.IsEmpty())
+	{
+		return false;
+	}
+
+	if (Property->IsA(FClassProperty::StaticClass()) || Property->IsA(FObjectPropertyBase::StaticClass()))
+	{
+		const bool bIsNull = FDefaultValueHelper::Is(CppForm, TEXT("NULL")) || FDefaultValueHelper::Is(CppForm, TEXT("nullptr")) || FDefaultValueHelper::Is(CppForm, TEXT("0"));
+		if (bIsNull)
+		{
+			OutForm = TEXT("None");
+		}
+		return bIsNull; // always return as null is the only the processing we can do for object defaults
+	}
+
+	auto ValidateEnumEntry = [Property, &CppForm](const UEnum* Enum, const FString& EnumValue)
+	{
+		const int32 EnumEntryIndex = Enum->GetIndexByName(*EnumValue);
+		if (EnumEntryIndex == INDEX_NONE)
+		{
+			return false;
+		}
+		if (Enum->HasMetaData(TEXT("Hidden"), EnumEntryIndex))
+		{
+			FError::Throwf(TEXT("Hidden enum entries cannot be used as default values: %s \"%s\" "), *Property->GetName(), *CppForm);
+		}
+		return true;
+	};
+
+	if( !Property->IsA(FStructProperty::StaticClass()) )
+	{
+		if( Property->IsA(FIntProperty::StaticClass()) )
+		{
+			int32 Value;
+			if( FDefaultValueHelper::ParseInt( CppForm, Value) )
+			{
+				OutForm = FString::FromInt(Value);
+			}
+		}
+		else if (Property->IsA(FInt64Property::StaticClass()))
+		{
+			int64 Value;
+			if (FDefaultValueHelper::ParseInt64(CppForm, Value))
+			{
+				OutForm = FString::Printf(TEXT("%lld"), Value);
+			}
+		}
+		else if( Property->IsA(FByteProperty::StaticClass()) )
+		{
+			const UEnum* Enum = CastFieldChecked<const FByteProperty>(const_cast<FProperty*>(Property))->Enum;
+			if( NULL != Enum )
+			{
+				OutForm = FDefaultValueHelper::GetUnqualifiedEnumValue(FDefaultValueHelper::RemoveWhitespaces(CppForm));
+				return ValidateEnumEntry(Enum, OutForm);
+			}
+			int32 Value;
+			if( FDefaultValueHelper::ParseInt( CppForm, Value) )
+			{
+				OutForm = FString::FromInt(Value);
+				return ( 0 <= Value ) && ( 255 >= Value );
+			}
+		}
+		else if( Property->IsA(FEnumProperty::StaticClass()) )
+		{
+			const FEnumProperty* EnumProp = CastFieldChecked<FEnumProperty>(const_cast<FProperty*>(Property));
+			if (const UEnum* Enum = CastFieldChecked<FEnumProperty>(const_cast<FProperty*>(Property))->GetEnum())
+			{
+				OutForm = FDefaultValueHelper::GetUnqualifiedEnumValue(FDefaultValueHelper::RemoveWhitespaces(CppForm));
+				return ValidateEnumEntry(Enum, OutForm);
+			}
+
+			int64 Value;
+			if (FDefaultValueHelper::ParseInt64(CppForm, Value))
+			{
+				OutForm = LexToString(Value);
+				return EnumProp->GetUnderlyingProperty()->CanHoldValue(Value);
+			}
+		}
+		else if( Property->IsA(FFloatProperty::StaticClass()) )
+		{
+			float Value;
+			if( FDefaultValueHelper::ParseFloat( CppForm, Value) )
+			{
+				OutForm = FString::Printf( TEXT("%f"), Value) ;
+			}
+		}
+		else if( Property->IsA(FDoubleProperty::StaticClass()) )
+		{
+			double Value;
+			if( FDefaultValueHelper::ParseDouble( CppForm, Value) )
+			{
+				OutForm = FString::Printf( TEXT("%f"), Value) ;
+			}
+		}
+		else if( Property->IsA(FBoolProperty::StaticClass()) )
+		{
+			if( FDefaultValueHelper::Is(CppForm, TEXT("true")) || 
+				FDefaultValueHelper::Is(CppForm, TEXT("false")) )
+			{
+				OutForm = FDefaultValueHelper::RemoveWhitespaces( CppForm );
+			}
+		}
+		else if( Property->IsA(FNameProperty::StaticClass()) )
+		{
+			if(FDefaultValueHelper::Is( CppForm, TEXT("NAME_None") ))
+			{
+				OutForm = TEXT("None");
+				return true;
+			}
+			return FDefaultValueHelper::StringFromCppString(CppForm, TEXT("FName"), OutForm);
+		}
+		else if( Property->IsA(FTextProperty::StaticClass()) )
+		{
+			// Handle legacy cases of FText::FromString being used as default values
+			// These should be replaced with INVTEXT as FText::FromString can produce inconsistent keys
+			if (FDefaultValueHelper::StringFromCppString(CppForm, TEXT("FText::FromString"), OutForm))
+			{
+				UE_LOG_WARNING_UHT(TEXT("FText::FromString should be replaced with INVTEXT for default parameter values"));
+				return true;
+			}
+
+			// Parse the potential value into an instance
+			FText ParsedText;
+			if (FDefaultValueHelper::Is(CppForm, TEXT("FText()")) || FDefaultValueHelper::Is(CppForm, TEXT("FText::GetEmpty()")))
+			{
+				ParsedText = FText::GetEmpty();
+			}
+			else
+			{
+				static const FString UHTDummyNamespace = TEXT("__UHT_DUMMY_NAMESPACE__");
+
+				if (!FTextStringHelper::ReadFromBuffer(*CppForm, ParsedText, *UHTDummyNamespace, nullptr, /*bRequiresQuotes*/true))
+				{
+					return false;
+				}
+
+				// If the namespace of the parsed text matches the default we gave then this was a LOCTEXT macro which we 
+				// don't allow in default values as they rely on an external macro that is known to C++ but not to UHT
+				// TODO: UHT could parse these if it tracked the current LOCTEXT_NAMESPACE macro as it parsed
+				if (TOptional<FString> ParsedTextNamespace = FTextInspector::GetNamespace(ParsedText))
+				{
+					if (ParsedTextNamespace.GetValue().Equals(UHTDummyNamespace))
+					{
+						FError::Throwf(TEXT("LOCTEXT default parameter values are not supported; use NSLOCTEXT instead: %s \"%s\" "), *Property->GetName(), *CppForm);
+					}
+				}
+			}
+
+			// Normalize the default value from the parsed value
+			FTextStringHelper::WriteToBuffer(OutForm, ParsedText, /*bRequiresQuotes*/false);
+			return true;
+		}
+		else if( Property->IsA(FStrProperty::StaticClass()) )
+		{
+			return FDefaultValueHelper::StringFromCppString(CppForm, TEXT("FString"), OutForm);
+		}
+	}
+	else 
+	{
+		// Cache off the struct types, in case we need them later
+		UPackage* CoreUObjectPackage = UObject::StaticClass()->GetOutermost();
+		static const UScriptStruct* VectorStruct = FindObjectChecked<UScriptStruct>(CoreUObjectPackage, TEXT("Vector"));
+		static const UScriptStruct* Vector2DStruct = FindObjectChecked<UScriptStruct>(CoreUObjectPackage, TEXT("Vector2D"));
+		static const UScriptStruct* RotatorStruct = FindObjectChecked<UScriptStruct>(CoreUObjectPackage, TEXT("Rotator"));
+		static const UScriptStruct* LinearColorStruct = FindObjectChecked<UScriptStruct>(CoreUObjectPackage, TEXT("LinearColor"));
+		static const UScriptStruct* ColorStruct = FindObjectChecked<UScriptStruct>(CoreUObjectPackage, TEXT("Color"));
+
+		const FStructProperty* StructProperty = CastFieldChecked<FStructProperty>(const_cast<FProperty*>(Property));
+		if( StructProperty->Struct == VectorStruct )
+		{
+			FString Parameters;
+			if(FDefaultValueHelper::Is( CppForm, TEXT("FVector::ZeroVector") ))
+			{
+				return true;
+			}
+			else if(FDefaultValueHelper::Is(CppForm, TEXT("FVector::UpVector")))
+			{
+				OutForm = FString::Printf(TEXT("%f,%f,%f"),
+					FVector::UpVector.X, FVector::UpVector.Y, FVector::UpVector.Z);
+			}
+			else if(FDefaultValueHelper::Is(CppForm, TEXT("FVector::ForwardVector")))
+			{
+				OutForm = FString::Printf(TEXT("%f,%f,%f"),
+					FVector::ForwardVector.X, FVector::ForwardVector.Y, FVector::ForwardVector.Z);
+			}
+			else if(FDefaultValueHelper::Is(CppForm, TEXT("FVector::RightVector")))
+			{
+				OutForm = FString::Printf(TEXT("%f,%f,%f"),
+					FVector::RightVector.X, FVector::RightVector.Y, FVector::RightVector.Z);
+			}
+			else if( FDefaultValueHelper::GetParameters(CppForm, TEXT("FVector"), Parameters) )
+			{
+				if( FDefaultValueHelper::Is(Parameters, TEXT("ForceInit")) )
+				{
+					return true;
+				}
+				FVector Vector;
+				float Value;
+				if (FDefaultValueHelper::ParseVector(Parameters, Vector))
+				{
+					OutForm = FString::Printf(TEXT("%f,%f,%f"),
+						Vector.X, Vector.Y, Vector.Z);
+				}
+				else if (FDefaultValueHelper::ParseFloat(Parameters, Value))
+				{
+					OutForm = FString::Printf(TEXT("%f,%f,%f"),
+						Value, Value, Value);
+				}
+			}
+		}
+		else if( StructProperty->Struct == RotatorStruct )
+		{
+			if(FDefaultValueHelper::Is( CppForm, TEXT("FRotator::ZeroRotator") ))
+			{
+				return true;
+			}
+			FString Parameters;
+			if( FDefaultValueHelper::GetParameters(CppForm, TEXT("FRotator"), Parameters) )
+			{
+				if( FDefaultValueHelper::Is(Parameters, TEXT("ForceInit")) )
+				{
+					return true;
+				}
+				FRotator Rotator;
+				if(FDefaultValueHelper::ParseRotator(Parameters, Rotator))
+				{
+					OutForm = FString::Printf(TEXT("%f,%f,%f"),
+						Rotator.Pitch, Rotator.Yaw, Rotator.Roll);
+				}
+			}
+		}
+		else if( StructProperty->Struct == Vector2DStruct )
+		{
+			if(FDefaultValueHelper::Is( CppForm, TEXT("FVector2D::ZeroVector") ))
+			{
+				return true;
+			}
+			if(FDefaultValueHelper::Is(CppForm, TEXT("FVector2D::UnitVector")))
+			{
+				OutForm = FString::Printf(TEXT("(X=%3.3f,Y=%3.3f)"),
+					FVector2D::UnitVector.X, FVector2D::UnitVector.Y);
+			}
+			FString Parameters;
+			if( FDefaultValueHelper::GetParameters(CppForm, TEXT("FVector2D"), Parameters) )
+			{
+				if( FDefaultValueHelper::Is(Parameters, TEXT("ForceInit")) )
+				{
+					return true;
+				}
+				FVector2D Vector2D;
+				if(FDefaultValueHelper::ParseVector2D(Parameters, Vector2D))
+				{
+					OutForm = FString::Printf(TEXT("(X=%3.3f,Y=%3.3f)"),
+						Vector2D.X, Vector2D.Y);
+				}
+			}
+		}
+		else if( StructProperty->Struct == LinearColorStruct )
+		{
+			if( FDefaultValueHelper::Is( CppForm, TEXT("FLinearColor::White") ) )
+			{
+				OutForm = FLinearColor::White.ToString();
+			}
+			else if ( FDefaultValueHelper::Is( CppForm, TEXT("FLinearColor::Gray") ) )
+			{
+				OutForm = FLinearColor::Gray.ToString();
+			}
+			else if ( FDefaultValueHelper::Is( CppForm, TEXT("FLinearColor::Black") ) )
+			{
+				OutForm = FLinearColor::Black.ToString();
+			}
+			else if ( FDefaultValueHelper::Is( CppForm, TEXT("FLinearColor::Transparent") ) )
+			{
+				OutForm = FLinearColor::Transparent.ToString();
+			}
+			else if ( FDefaultValueHelper::Is( CppForm, TEXT("FLinearColor::Red") ) )
+			{
+				OutForm = FLinearColor::Red.ToString();
+			}
+			else if ( FDefaultValueHelper::Is( CppForm, TEXT("FLinearColor::Green") ) )
+			{
+				OutForm = FLinearColor::Green.ToString();
+			}
+			else if ( FDefaultValueHelper::Is( CppForm, TEXT("FLinearColor::Blue") ) )
+			{
+				OutForm = FLinearColor::Blue.ToString();
+			}
+			else if ( FDefaultValueHelper::Is( CppForm, TEXT("FLinearColor::Yellow") ) )
+			{
+				OutForm = FLinearColor::Yellow.ToString();
+			}
+			else
+			{
+				FString Parameters;
+				if( FDefaultValueHelper::GetParameters(CppForm, TEXT("FLinearColor"), Parameters) )
+				{
+					if( FDefaultValueHelper::Is(Parameters, TEXT("ForceInit")) )
+					{
+						return true;
+					}
+					FLinearColor Color;
+					if( FDefaultValueHelper::ParseLinearColor(Parameters, Color) )
+					{
+						OutForm = Color.ToString();
+					}
+				}
+			}
+		}
+		else if( StructProperty->Struct == ColorStruct )
+		{
+			if( FDefaultValueHelper::Is( CppForm, TEXT("FColor::White") ) )
+			{
+				OutForm = FColor::White.ToString();
+			}
+			else if ( FDefaultValueHelper::Is( CppForm, TEXT("FColor::Black") ) )
+			{
+				OutForm = FColor::Black.ToString();
+			}
+			else if ( FDefaultValueHelper::Is( CppForm, TEXT("FColor::Red") ) )
+			{
+				OutForm = FColor::Red.ToString();
+			}
+			else if ( FDefaultValueHelper::Is( CppForm, TEXT("FColor::Green") ) )
+			{
+				OutForm = FColor::Green.ToString();
+			}
+			else if ( FDefaultValueHelper::Is( CppForm, TEXT("FColor::Blue") ) )
+			{
+				OutForm = FColor::Blue.ToString();
+			}
+			else if (FDefaultValueHelper::Is(CppForm, TEXT("FColor::Yellow")))
+			{
+				OutForm = FColor::Yellow.ToString();
+			}
+			else if ( FDefaultValueHelper::Is( CppForm, TEXT("FColor::Cyan") ) )
+			{
+				OutForm = FColor::Cyan.ToString();
+			}
+			else if ( FDefaultValueHelper::Is( CppForm, TEXT("FColor::Magenta") ) )
+			{
+				OutForm = FColor::Magenta.ToString();
+			}
+			else
+			{
+				FString Parameters;
+				if( FDefaultValueHelper::GetParameters(CppForm, TEXT("FColor"), Parameters) )
+				{
+					if( FDefaultValueHelper::Is(Parameters, TEXT("ForceInit")) )
+					{
+						return true;
+					}
+					FColor Color;
+					if( FDefaultValueHelper::ParseColor(Parameters, Color) )
+					{
+						OutForm = Color.ToString();
+					}
+				}
+			}
+		}
+	}
+
+	return !OutForm.IsEmpty();
+}
+
+bool FHeaderParser::TryToMatchConstructorParameterList(FToken Token)
+{
+	FToken PotentialParenthesisToken;
+	if (!GetToken(PotentialParenthesisToken))
+	{
+		return false;
+	}
+
+	if (!PotentialParenthesisToken.Matches(TEXT('(')))
+	{
+		UngetToken(PotentialParenthesisToken);
+		return false;
+	}
+
+	FClassMetaData* ClassData = GScriptHelper.FindClassData(GetCurrentClass());
+	check(ClassData);
+
+	bool bOICtor = false;
+	bool bVTCtor = false;
+
+	if (!ClassData->bDefaultConstructorDeclared && MatchSymbol(TEXT(')')))
+	{
+		ClassData->bDefaultConstructorDeclared = true;
+	}
+	else if (!ClassData->bObjectInitializerConstructorDeclared
+		|| !ClassData->bCustomVTableHelperConstructorDeclared
+	)
+	{
+		FToken ObjectInitializerParamParsingToken;
+
+		bool bIsConst = false;
+		bool bIsRef = false;
+		int32 ParenthesesNestingLevel = 1;
+
+		while (ParenthesesNestingLevel && GetToken(ObjectInitializerParamParsingToken))
+		{
+			// Template instantiation or additional parameter excludes ObjectInitializer constructor.
+			if (ObjectInitializerParamParsingToken.Matches(TEXT(',')) || ObjectInitializerParamParsingToken.Matches(TEXT('<')))
+			{
+				bOICtor = false;
+				bVTCtor = false;
+				break;
+			}
+
+			if (ObjectInitializerParamParsingToken.Matches(TEXT('(')))
+			{
+				ParenthesesNestingLevel++;
+				continue;
+			}
+
+			if (ObjectInitializerParamParsingToken.Matches(TEXT(')')))
+			{
+				ParenthesesNestingLevel--;
+				continue;
+			}
+
+			if (ObjectInitializerParamParsingToken.Matches(TEXT("const"), ESearchCase::CaseSensitive))
+			{
+				bIsConst = true;
+				continue;
+			}
+
+			if (ObjectInitializerParamParsingToken.Matches(TEXT('&')))
+			{
+				bIsRef = true;
+				continue;
+			}
+
+			if (ObjectInitializerParamParsingToken.Matches(TEXT("FObjectInitializer"), ESearchCase::CaseSensitive)
+				|| ObjectInitializerParamParsingToken.Matches(TEXT("FPostConstructInitializeProperties"), ESearchCase::CaseSensitive) // Deprecated, but left here, so it won't break legacy code.
+				)
+			{
+				bOICtor = true;
+			}
+
+			if (ObjectInitializerParamParsingToken.Matches(TEXT("FVTableHelper"), ESearchCase::CaseSensitive))
+			{
+				bVTCtor = true;
+			}
+		}
+
+		// Parse until finish.
+		while (ParenthesesNestingLevel && GetToken(ObjectInitializerParamParsingToken))
+		{
+			if (ObjectInitializerParamParsingToken.Matches(TEXT('(')))
+			{
+				ParenthesesNestingLevel++;
+				continue;
+			}
+
+			if (ObjectInitializerParamParsingToken.Matches(TEXT(')')))
+			{
+				ParenthesesNestingLevel--;
+				continue;
+			}
+		}
+
+		ClassData->bObjectInitializerConstructorDeclared = ClassData->bObjectInitializerConstructorDeclared || (bOICtor && bIsRef && bIsConst);
+		ClassData->bCustomVTableHelperConstructorDeclared = ClassData->bCustomVTableHelperConstructorDeclared || (bVTCtor && bIsRef);
+	}
+
+	ClassData->bConstructorDeclared = ClassData->bConstructorDeclared || !bVTCtor;
+
+	// Optionally match semicolon.
+	if (!MatchSymbol(TEXT(';')))
+	{
+		// If not matched a semicolon, this is inline constructor definition. We have to skip it.
+		UngetToken(Token); // Resets input stream to the initial token.
+		GetToken(Token); // Re-gets the initial token to start constructor definition skip.
+		return SkipDeclaration(Token);
+	}
+
+	return true;
+}
+
+void FHeaderParser::CompileVersionDeclaration(UStruct* Struct)
+{
+	FUnrealSourceFile* CurrentSourceFilePtr = GetCurrentSourceFile();
+	TSharedRef<FUnrealSourceFile> CurrentSrcFile = CurrentSourceFilePtr->AsShared();
+	// Do nothing if we're at the end of file.
+	FToken Token;
+	if (!GetToken(Token, true, ESymbolParseOption::Normal))
+	{
+		return;
+	}
+
+	// Default version based on config file.
+	EGeneratedCodeVersion Version = DefaultGeneratedCodeVersion;
+
+	// Overwrite with module-specific value if one was specified.
+	if (CurrentlyParsedModule->GeneratedCodeVersion != EGeneratedCodeVersion::None)
+	{
+		Version = CurrentlyParsedModule->GeneratedCodeVersion;
+	}
+
+	if (Token.TokenType == ETokenType::TOKEN_Symbol
+		&& Token.Matches(TEXT(')')))
+	{
+		CurrentSrcFile->GetGeneratedCodeVersions().FindOrAdd(Struct) = Version;
+		UngetToken(Token);
+		return;
+	}
+
+	// Overwrite with version specified by macro.
+	Version = ToGeneratedCodeVersion(Token.Identifier);
+
+	CurrentSrcFile->GetGeneratedCodeVersions().FindOrAdd(Struct) = Version;
+}
+
+void FHeaderParser::ResetClassData()
+{
+	UClass* CurrentClass = GetCurrentClass();
+	CurrentClass->PropertiesSize = 0;
+
+	// Set class flags and within.
+	CurrentClass->ClassFlags &= ~CLASS_RecompilerClear;
+
+	if (UClass* SuperClass = CurrentClass->GetSuperClass())
+	{
+		CurrentClass->ClassFlags |= (SuperClass->ClassFlags) & CLASS_ScriptInherit;
+		CurrentClass->ClassConfigName = SuperClass->ClassConfigName;
+		check(SuperClass->ClassWithin);
+		if (CurrentClass->ClassWithin == nullptr)
+		{
+			CurrentClass->ClassWithin = SuperClass->ClassWithin;
+		}
+
+		// Copy special categories from parent
+		if (SuperClass->HasMetaData(FHeaderParserNames::NAME_HideCategories))
+		{
+			CurrentClass->SetMetaData(FHeaderParserNames::NAME_HideCategories, *SuperClass->GetMetaData(FHeaderParserNames::NAME_HideCategories));
+		}
+		if (SuperClass->HasMetaData(FHeaderParserNames::NAME_ShowCategories))
+		{
+			CurrentClass->SetMetaData(FHeaderParserNames::NAME_ShowCategories, *SuperClass->GetMetaData(FHeaderParserNames::NAME_ShowCategories));
+		}
+		if (SuperClass->HasMetaData(FHeaderParserNames::NAME_SparseClassDataTypes))
+		{
+			CurrentClass->SetMetaData(FHeaderParserNames::NAME_SparseClassDataTypes, *SuperClass->GetMetaData(FHeaderParserNames::NAME_SparseClassDataTypes));
+		}
+		if (SuperClass->HasMetaData(NAME_HideFunctions))
+		{
+			CurrentClass->SetMetaData(NAME_HideFunctions, *SuperClass->GetMetaData(NAME_HideFunctions));
+		}
+		if (SuperClass->HasMetaData(NAME_AutoExpandCategories))
+		{
+			CurrentClass->SetMetaData(NAME_AutoExpandCategories, *SuperClass->GetMetaData(NAME_AutoExpandCategories));
+		}
+		if (SuperClass->HasMetaData(NAME_AutoCollapseCategories))
+		{
+			CurrentClass->SetMetaData(NAME_AutoCollapseCategories, *SuperClass->GetMetaData(NAME_AutoCollapseCategories));
+		}
+	}
+
+	check(CurrentClass->ClassWithin);
+}
+
+void FHeaderParser::PostPopNestClass(UClass* CurrentClass)
+{
+	// Validate all the rep notify events here, to make sure they're implemented
+	VerifyPropertyMarkups(CurrentClass);
+
+	// Iterate over all the interfaces we claim to implement
+	for (FImplementedInterface& Impl : CurrentClass->Interfaces)
+	{
+		// And their super-classes
+		for (UClass* Interface = Impl.Class; Interface; Interface = Interface->GetSuperClass())
+		{
+			// If this interface is a common ancestor, skip it
+			if (CurrentClass->IsChildOf(Interface))
+			{
+				continue;
+			}
+
+			// So iterate over all functions this interface declares
+			for (UFunction* InterfaceFunction : TFieldRange<UFunction>(Interface, EFieldIteratorFlags::ExcludeSuper))
+			{
+				bool bImplemented = false;
+
+				// And try to find one that matches
+				for (UFunction* ClassFunction : TFieldRange<UFunction>(CurrentClass))
+				{
+					if (ClassFunction->GetFName() != InterfaceFunction->GetFName())
+					{
+						continue;
+					}
+
+					if ((InterfaceFunction->FunctionFlags & FUNC_Event) && !(ClassFunction->FunctionFlags & FUNC_Event))
+					{
+						FError::Throwf(TEXT("Implementation of function '%s::%s' must be declared as 'event' to match declaration in interface '%s'"), *ClassFunction->GetOuter()->GetName(), *ClassFunction->GetName(), *Interface->GetName());
+					}
+
+					if ((InterfaceFunction->FunctionFlags & FUNC_Delegate) && !(ClassFunction->FunctionFlags & FUNC_Delegate))
+					{
+						FError::Throwf(TEXT("Implementation of function '%s::%s' must be declared as 'delegate' to match declaration in interface '%s'"), *ClassFunction->GetOuter()->GetName(), *ClassFunction->GetName(), *Interface->GetName());
+					}
+
+					// Making sure all the parameters match up correctly
+					bImplemented = true;
+
+					if (ClassFunction->NumParms != InterfaceFunction->NumParms)
+					{
+						FError::Throwf(TEXT("Implementation of function '%s' conflicts with interface '%s' - different number of parameters (%i/%i)"), *InterfaceFunction->GetName(), *Interface->GetName(), ClassFunction->NumParms, InterfaceFunction->NumParms);
+					}
+
+					int32 Count = 0;
+					for (TFieldIterator<FProperty> It1(InterfaceFunction), It2(ClassFunction); Count < ClassFunction->NumParms; ++It1, ++It2, Count++)
+					{
+						if (!FPropertyBase(*It1).MatchesType(FPropertyBase(*It2), 1))
+						{
+							if (It1->PropertyFlags & CPF_ReturnParm)
+							{
+								FError::Throwf(TEXT("Implementation of function '%s' conflicts only by return type with interface '%s'"), *InterfaceFunction->GetName(), *Interface->GetName());
+							}
+							else
+							{
+								FError::Throwf(TEXT("Implementation of function '%s' conflicts with interface '%s' - parameter %i '%s'"), *InterfaceFunction->GetName(), *Interface->GetName(), Count, *It1->GetName());
+							}
+						}
+					}
+				}
+
+				// Delegate signature functions are simple stubs and aren't required to be implemented (they are not callable)
+				if (InterfaceFunction->FunctionFlags & FUNC_Delegate)
+				{
+					bImplemented = true;
+				}
+
+				// Verify that if this has blueprint-callable functions that are not implementable events, we've implemented them as a UFunction in the target class
+				if (!bImplemented
+					&& InterfaceFunction->HasAnyFunctionFlags(FUNC_BlueprintCallable)
+					&& !InterfaceFunction->HasAnyFunctionFlags(FUNC_BlueprintEvent)
+					&& !Interface->HasMetaData(NAME_CannotImplementInterfaceInBlueprint))  // FBlueprintMetadata::MD_CannotImplementInterfaceInBlueprint
+				{
+					FError::Throwf(TEXT("Missing UFunction implementation of function '%s' from interface '%s'.  This function needs a UFUNCTION() declaration."), *InterfaceFunction->GetName(), *Interface->GetName());
+				}
+			}
+		}
+	}
+}
+
+void FHeaderParser::PostPopFunctionDeclaration(FClasses& AllClasses, UFunction* PoppedFunction)
+{
+	//@TODO: UCREMOVAL: Move this code to occur at delegate var declaration, and force delegates to be declared before variables that use them
+	if (!GetCurrentScope()->IsFileScope() && GetCurrentClassData()->ContainsDelegates())
+	{
+		// now validate all delegate variables declared in the class
+		TMap<FName, UFunction*> DelegateCache;
+		FixupDelegateProperties(AllClasses, PoppedFunction, *GetCurrentScope(), DelegateCache);
+	}
+}
+
+void FHeaderParser::PostPopNestInterface(FClasses& AllClasses, UClass* CurrentInterface)
+{
+	FClassMetaData* ClassData = GScriptHelper.FindClassData(CurrentInterface);
+	check(ClassData);
+	if (ClassData->ContainsDelegates())
+	{
+		TMap<FName, UFunction*> DelegateCache;
+		FixupDelegateProperties(AllClasses, CurrentInterface, FScope::GetTypeScope(ExactCast<UClass>(CurrentInterface)).Get(), DelegateCache);
+	}
+}
+
+FDocumentationPolicy FHeaderParser::GetDocumentationPolicyFromName(const FString& PolicyName)
+{
+	FDocumentationPolicy DocumentationPolicy;
+	if (FCString::Strcmp(*PolicyName, TEXT("Strict")) == 0)
+	{
+		DocumentationPolicy.bClassOrStructCommentRequired = true;
+		DocumentationPolicy.bFunctionToolTipsRequired = true;
+		DocumentationPolicy.bMemberToolTipsRequired = true;
+		DocumentationPolicy.bParameterToolTipsRequired = true;
+		DocumentationPolicy.bFloatRangesRequired = true;
+	}
+	else
+	{
+		FError::Throwf(TEXT("Documentation Policy '%s' not yet supported"), *PolicyName);
+	}
+	return DocumentationPolicy;
+}
+
+
+FDocumentationPolicy FHeaderParser::GetDocumentationPolicyForStruct(UStruct* Struct)
+{
+	SCOPE_SECONDS_COUNTER_UHT(DocumentationPolicy);
+
+	check(Struct!= nullptr);
+
+	FDocumentationPolicy DocumentationPolicy;
+	FString DocumentationPolicyName;
+	if (Struct->GetStringMetaDataHierarchical(NAME_DocumentationPolicy, &DocumentationPolicyName))
+	{
+		DocumentationPolicy = GetDocumentationPolicyFromName(DocumentationPolicyName);
+	}
+	return DocumentationPolicy;
+}
+
+void FHeaderParser::CheckDocumentationPolicyForEnum(UEnum* Enum, const TMap<FName, FString>& MetaData, const TArray<TMap<FName, FString>>& Entries)
+{
+	SCOPE_SECONDS_COUNTER_UHT(DocumentationPolicy);
+
+	check(Enum != nullptr);
+
+	const FString* DocumentationPolicyName = MetaData.Find(NAME_DocumentationPolicy);
+	if (DocumentationPolicyName == nullptr)
+	{
+		return;
+	}
+
+	check(!DocumentationPolicyName->IsEmpty());
+
+	FDocumentationPolicy DocumentationPolicy = GetDocumentationPolicyFromName(*DocumentationPolicyName);
+	if (DocumentationPolicy.bClassOrStructCommentRequired)
+	{
+		const FString* EnumToolTip = MetaData.Find(NAME_ToolTip);
+		if (EnumToolTip == nullptr)
+		{
+			UE_LOG_ERROR_UHT(TEXT("Enum '%s' does not provide a tooltip / comment (DocumentationPolicy)."), *Enum->GetName());
+		}
+	}
+
+	TMap<FString, FString> ToolTipToEntry;
+	for (const TMap<FName, FString>& Entry : Entries)
+	{
+		const FString* EntryName = Entry.Find(NAME_Name);
+		if (EntryName == nullptr)
+		{
+			continue;
+		}
+
+		const FString* ToolTip = Entry.Find(NAME_ToolTip);
+		if (ToolTip == nullptr)
+		{
+			UE_LOG_ERROR_UHT(TEXT("Enum entry '%s::%s' does not provide a tooltip / comment (DocumentationPolicy)."), *Enum->GetName(), *(*EntryName));
+			continue;
+		}
+
+		const FString* ExistingEntry = ToolTipToEntry.Find(*ToolTip);
+		if (ExistingEntry != nullptr)
+		{
+			UE_LOG_ERROR_UHT(TEXT("Enum entries '%s::%s' and '%s::%s' have identical tooltips / comments (DocumentationPolicy)."), *Enum->GetName(), *(*ExistingEntry), *Enum->GetName(), *(*EntryName));
+		}
+		ToolTipToEntry.Add(*ToolTip, *EntryName);
+	}
+}
+
+void FHeaderParser::CheckDocumentationPolicyForStruct(UStruct* Struct)
+{
+	SCOPE_SECONDS_COUNTER_UHT(DocumentationPolicy);
+
+	check(Struct != nullptr);
+
+	FDocumentationPolicy DocumentationPolicy = GetDocumentationPolicyForStruct(Struct);
+	if (DocumentationPolicy.bClassOrStructCommentRequired)
+	{
+		FString ClassTooltip;
+		if (const FString* ClassTooltipPtr = Struct->FindMetaData(NAME_ToolTip))
+		{
+			ClassTooltip = *ClassTooltipPtr;
+		}
+
+		if (ClassTooltip.IsEmpty() || ClassTooltip.Equals(Struct->GetName()))
+		{
+			UE_LOG_ERROR_UHT(TEXT("Struct '%s' does not provide a tooltip / comment (DocumentationPolicy)."), *Struct->GetName());
+		}
+	}
+
+	if (DocumentationPolicy.bMemberToolTipsRequired)
+	{
+		TMap<FString, FName> ToolTipToPropertyName;
+		for (FProperty* Property : TFieldRange<FProperty>(Struct, EFieldIteratorFlags::ExcludeSuper))
+		{
+			FString ToolTip = Property->GetToolTipText().ToString();
+			if (ToolTip.IsEmpty() || ToolTip.Equals(Property->GetDisplayNameText().ToString()))
+			{
+				UE_LOG_ERROR_UHT(TEXT("Property '%s::%s' does not provide a tooltip / comment (DocumentationPolicy)."), *Struct->GetName(), *Property->GetName());
+				continue;
+			}
+			const FName* ExistingPropertyName = ToolTipToPropertyName.Find(ToolTip);
+			if (ExistingPropertyName != nullptr)
+			{
+				UE_LOG_ERROR_UHT(TEXT("Property '%s::%s' and '%s::%s' are using identical tooltips (DocumentationPolicy)."), *Struct->GetName(), *ExistingPropertyName->ToString(), *Struct->GetName(), *Property->GetName());
+			}
+			ToolTipToPropertyName.Add(MoveTemp(ToolTip), Property->GetFName());
+		}
+	}
+
+	if (DocumentationPolicy.bFloatRangesRequired)
+	{
+		for (FProperty* Property : TFieldRange<FProperty>(Struct, EFieldIteratorFlags::ExcludeSuper))
+		{
+			if(DoesCPPTypeRequireDocumentation(Property->GetCPPType()))
+			{
+				const FString& UIMin = Property->GetMetaData(NAME_UIMin);
+				const FString& UIMax = Property->GetMetaData(NAME_UIMax);
+
+				if(!CheckUIMinMaxRangeFromMetaData(UIMin, UIMax))
+				{
+					UE_LOG_ERROR_UHT(TEXT("Property '%s::%s' does not provide a valid UIMin / UIMax (DocumentationPolicy)."), *Struct->GetName(), *Property->GetName());
+				}
+			}
+		}
+	}
+
+	// also compare all tooltips to see if they are unique
+	if (DocumentationPolicy.bFunctionToolTipsRequired)
+	{
+		UClass* Class = Cast<UClass>(Struct);
+		if (Class != nullptr)
+		{
+			TMap<FString, FName> ToolTipToFunc;
+			for (UFunction* Func : TFieldRange<UFunction>(Class, EFieldIteratorFlags::ExcludeSuper))
+			{
+				FString ToolTip = Func->GetToolTipText().ToString();
+				if (ToolTip.IsEmpty())
+				{
+					UE_LOG_ERROR_UHT(TEXT("Function '%s::%s' does not provide a tooltip / comment (DocumentationPolicy)."), *Class->GetName(), *Func->GetName());
+					continue;
+				}
+				const FName* ExistingFuncName = ToolTipToFunc.Find(ToolTip);
+				if (ExistingFuncName != nullptr)
+				{
+					UE_LOG_ERROR_UHT(TEXT("Functions '%s::%s' and '%s::%s' uses identical tooltips / comments (DocumentationPolicy)."), *Class->GetName(), *(*ExistingFuncName).ToString(), *Class->GetName(), *Func->GetName());
+				}
+				ToolTipToFunc.Add(MoveTemp(ToolTip), Func->GetFName());
+			}
+		}
+	}
+}
+
+bool FHeaderParser::DoesCPPTypeRequireDocumentation(const FString& CPPType)
+{
+	return PropertyCPPTypesRequiringUIRanges.Contains(CPPType);
+}
+
+// Validates the documentation for a given method
+void FHeaderParser::CheckDocumentationPolicyForFunc(UClass* Class, UFunction* Func)
+{
+	SCOPE_SECONDS_COUNTER_UHT(DocumentationPolicy);
+
+	check(Class != nullptr);
+	check(Func != nullptr);
+
+	FDocumentationPolicy DocumentationPolicy = GetDocumentationPolicyForStruct(Class);
+	if (DocumentationPolicy.bFunctionToolTipsRequired)
+	{
+		const FString* FunctionTooltip = Func->FindMetaData(NAME_ToolTip);
+		if (FunctionTooltip == nullptr)
+		{
+			UE_LOG_ERROR_UHT(TEXT("Function '%s::%s' does not provide a tooltip / comment (DocumentationPolicy)."), *Class->GetName(), *Func->GetName());
+		}
+	}
+
+	if (DocumentationPolicy.bParameterToolTipsRequired)
+	{
+		const FString* FunctionComment = Func->FindMetaData(NAME_Comment);
+		if (FunctionComment == nullptr)
+		{
+			UE_LOG_ERROR_UHT(TEXT("Function '%s::%s' does not provide a comment (DocumentationPolicy)."), *Class->GetName(), *Func->GetName());
+			return;
+		}
+		
+		TMap<FName, FString> ParamToolTips = GetParameterToolTipsFromFunctionComment(*FunctionComment);
+		bool HasAnyParamToolTips = ParamToolTips.Num() > 0;
+		if (ParamToolTips.Num() == 0)
+		{
+			const FString* ReturnValueToolTip = ParamToolTips.Find(NAME_ReturnValue);
+			if (ReturnValueToolTip != nullptr)
+			{
+				HasAnyParamToolTips = false;
+			}
+		}
+
+		// only apply the validation for parameter tooltips if a function has any @param statements at all.
+		if (HasAnyParamToolTips)
+		{
+			// ensure each parameter has a tooltip
+			TSet<FName> ExistingFields;
+			for (FProperty* Property : TFieldRange<FProperty>(Func))
+			{
+				FName ParamName = Property->GetFName();
+				if (ParamName == NAME_ReturnValue)
+				{
+					continue;
+				}
+				const FString* ParamToolTip = ParamToolTips.Find(ParamName);
+				if (ParamToolTip == nullptr)
+				{
+					UE_LOG_ERROR_UHT(TEXT("Function '%s::%s' doesn't provide a tooltip for parameter '%s' (DocumentationPolicy)."), *Class->GetName(), *Func->GetName(), *ParamName.ToString());
+				}
+				ExistingFields.Add(ParamName);
+			}
+
+			// ensure we don't have parameter tooltips for parameters that don't exist
+			for (TPair<FName, FString>& Pair : ParamToolTips)
+			{
+				const FName& ParamName = Pair.Key;
+				if (ParamName == NAME_ReturnValue)
+				{
+					continue;
+				}
+				if (!ExistingFields.Contains(ParamName))
+				{
+					UE_LOG_ERROR_UHT(TEXT("Function '%s::%s' provides a tooltip for an unknown parameter '%s' (DocumentationPolicy)."), *Class->GetName(), *Func->GetName(), *Pair.Key.ToString());
+				}
+			}
+
+			// check for duplicate tooltips
+			TMap<FString, FName> ToolTipToParam;
+			for (TPair<FName, FString>& Pair : ParamToolTips)
+			{
+				const FName& ParamName = Pair.Key;
+				if (ParamName == NAME_ReturnValue)
+				{
+					continue;
+				}
+				const FName* ExistingParam = ToolTipToParam.Find(Pair.Value);
+				if (ExistingParam != nullptr)
+				{
+					UE_LOG_ERROR_UHT(TEXT("Function '%s::%s' uses identical tooltips for parameters '%s' and '%s' (DocumentationPolicy)."), *Class->GetName(), *Func->GetName(), *ExistingParam->ToString(), *Pair.Key.ToString());
+				}
+				ToolTipToParam.Add(MoveTemp(Pair.Value), Pair.Key);
+			}
+		}
+	}
+}
+
+bool FHeaderParser::IsReservedTypeName(const FString& TypeName)
+{
+	for(const FString& ReservedName : ReservedTypeNames)
+	{
+		if(TypeName == ReservedName)
+		{
+			return true;
+		}
+	}
+	return false;
+}
+
+bool FHeaderParser::IsReservedTypeName(const FToken& Token)
+{
+	for (const FString& ReservedName : ReservedTypeNames)
+	{
+		if (Token.Matches(*ReservedName, ESearchCase::IgnoreCase))
+		{
+			return true;
+		}
+	}
+	return false;
+}
+
+bool FHeaderParser::CheckUIMinMaxRangeFromMetaData(const FString& UIMin, const FString& UIMax)
+{
+	if (UIMin.IsEmpty() || UIMax.IsEmpty())
+	{
+		return false;
+	}
+
+	double UIMinValue = FCString::Atod(*UIMin);
+	double UIMaxValue = FCString::Atod(*UIMax);
+	if (UIMin > UIMax) // note that we actually allow UIMin == UIMax to disable the range manually.
+	{
+		return false;
+	}
+
+	return true;
+}
+
+template <class TFunctionType>
+TFunctionType* CreateFunctionImpl(const FFuncInfo& FuncInfo, UObject* Outer, FScope* CurrentScope)
+{
+	// Allocate local property frame, push nesting level and verify
+	// uniqueness at this scope level.
+	{
+		auto TypeIterator = CurrentScope->GetTypeIterator();
+		while (TypeIterator.MoveNext())
+		{
+			UField* Type = *TypeIterator;
+			if (Type->GetFName() == FuncInfo.Function.Identifier)
+			{
+				FError::Throwf(TEXT("'%s' conflicts with '%s'"), FuncInfo.Function.Identifier, *Type->GetFullName());
+			}
+		}
+	}
+
+	TFunctionType* Function = new(EC_InternalUseOnlyConstructor, Outer, FuncInfo.Function.Identifier, RF_Public) TFunctionType(FObjectInitializer(), nullptr);
+	Function->ReturnValueOffset = MAX_uint16;
+	Function->FirstPropertyToInit = nullptr;
+
+	if (!CurrentScope->IsFileScope())
+	{
+		UStruct* Struct = ((FStructScope*)CurrentScope)->GetStruct();
+
+		Function->Next = Struct->Children;
+		Struct->Children = Function;
+	}
+
+	return Function;
+}
+
+UFunction* FHeaderParser::CreateFunction(const FFuncInfo &FuncInfo) const
+{
+	return CreateFunctionImpl<UFunction>(FuncInfo, GetCurrentClass(), GetCurrentScope());
+}
+
+template<class T>
+UDelegateFunction* FHeaderParser::CreateDelegateFunction(const FFuncInfo &FuncInfo) const
+{
+	FFileScope* CurrentFileScope = GetCurrentFileScope();
+	FUnrealSourceFile* LocSourceFile = CurrentFileScope ? CurrentFileScope->GetSourceFile() : nullptr;
+	UObject* CurrentPackage = LocSourceFile ? LocSourceFile->GetPackage() : nullptr;
+	return CreateFunctionImpl<T>(FuncInfo, IsInAClass() ? (UObject*)GetCurrentClass() : CurrentPackage, GetCurrentScope());
+}
