Branching UHT to UHTLite:

diff --git a/Engine/Source/Programs/UHTLite/Private/HeaderParser.h b/Engine/Source/Programs/UHTLite/Private/HeaderParser.h
new file mode 100644
--- /dev/null
+++ b/Engine/Source/Programs/UHTLite/Private/HeaderParser.h
@@ -0,0 +1,1089 @@
+// Copyright Epic Games, Inc. All Rights Reserved.
+
+#pragma once
+
+#include "CoreMinimal.h"
+#include "ParserHelper.h"
+#include "BaseParser.h"
+#include "Misc/CompilationResult.h"
+#include "Scope.h"
+
+class UClass;
+enum class EGeneratedCodeVersion : uint8;
+class FFeedbackContext;
+class UPackage;
+struct FManifestModule;
+class IScriptGeneratorPluginInterface;
+class FStringOutputDevice;
+class FProperty;
+class FUnrealSourceFile;
+class UFunction;
+class UEnum;
+class UScriptStruct;
+class UDelegateFunction;
+class UStruct;
+class FClass;
+class FClasses;
+class FScope;
+class FHeaderProvider;
+
+extern double GPluginOverheadTime;
+extern double GHeaderCodeGenTime;
+
+/*-----------------------------------------------------------------------------
+	Constants & types.
+-----------------------------------------------------------------------------*/
+
+enum {MAX_NEST_LEVELS = 16};
+
+/* Code nesting types. */
+enum class ENestType
+{
+	GlobalScope,
+	Class,
+	FunctionDeclaration,
+	Interface,
+	NativeInterface
+};
+
+/** Types of statements to allow within a particular nesting block. */
+enum class ENestAllowFlags
+{
+	None                 = 0,
+	Function             = 1,  // Allow Event declarations at this level.
+	VarDecl              = 2,  // Allow variable declarations at this level.
+	Class                = 4,  // Allow class definition heading.
+	Return               = 8,  // Allow 'return' within a function.
+	TypeDecl             = 16, // Allow declarations which do not affect memory layout, such as structs, enums, and consts, but not implicit delegates
+	ImplicitDelegateDecl = 32, // Allow implicit delegates (i.e. those not decorated with UDELEGATE) to be declared
+};
+
+ENUM_CLASS_FLAGS(ENestAllowFlags)
+
+namespace EDelegateSpecifierAction
+{
+	enum Type
+	{
+		DontParse,
+		Parse
+	};
+}
+
+/** The category of variable declaration being parsed */
+namespace EVariableCategory
+{
+	enum Type
+	{
+		RegularParameter,
+		ReplicatedParameter,
+		Return,
+		Member
+	};
+}
+
+enum class ELayoutMacroType
+{
+	None = -1,
+	Array,
+	ArrayEditorOnly,
+	Bitfield,
+	BitfieldEditorOnly,
+	Field,
+	FieldEditorOnly,
+	FieldInitialized,
+};
+
+/** Information for a particular nesting level. */
+class FNestInfo
+{
+	/** Link to the stack node. */
+	FScope* Scope;
+
+public:
+	/**
+	 * Gets nesting scope.
+	 */
+	FScope* GetScope() const
+	{
+		return Scope;
+	}
+
+	/**
+	 * Sets nesting scope.
+	 */
+	void SetScope(FScope* InScope)
+	{
+		this->Scope = InScope;
+	}
+
+	/** Statement that caused the nesting. */
+	ENestType NestType;
+
+	/** Types of statements to allow at this nesting level. */
+	ENestAllowFlags Allow;
+};
+
+struct FIndexRange
+{
+	int32 StartIndex;
+	int32 Count;
+};
+
+/**
+ * The FRigVMParameter represents a single parameter of a method
+ * marked up with RIGVM_METHOD.
+ * Each parameter can be marked with Constant, Input or Output
+ * metadata - this struct simplifies access to that information.
+ */
+struct FRigVMParameter
+{
+	FRigVMParameter()
+		: Name()
+		, Type()
+		, bConstant(false)
+		, bInput(false)
+		, bOutput(false)
+		, MaxArraySize()
+		, Getter()
+		, CastName()
+		, CastType()
+		, bEditorOnly(false)
+	{
+	}
+
+	FString Name;
+	FString Type;
+	bool bConstant;
+	bool bInput;
+	bool bOutput;
+	FString MaxArraySize;
+	FString Getter;
+	FString CastName;
+	FString CastType;
+	bool bEditorOnly;
+
+	const FString& NameOriginal(bool bCastName = false) const
+	{
+		return (bCastName && !CastName.IsEmpty()) ? CastName : Name;
+	}
+
+	const FString& TypeOriginal(bool bCastType = false) const
+	{
+		return (bCastType && !CastType.IsEmpty()) ? CastType : Type;
+	}
+
+	FString Declaration(bool bCastType = false, bool bCastName = false) const
+	{
+		return FString::Printf(TEXT("%s %s"), *TypeOriginal(bCastType), *NameOriginal(bCastName));
+	}
+
+	FString BaseType(bool bCastType = false) const
+	{
+		const FString& String = TypeOriginal(bCastType);
+		int32 LesserPos = 0;
+		if (String.FindChar('<', LesserPos))
+		{
+			return String.Mid(0, LesserPos);
+		}
+		return String;
+	}
+
+	FString ExtendedType(bool bCastType = false) const
+	{
+		const FString& String = TypeOriginal(bCastType);
+		int32 LesserPos = 0;
+		if(String.FindChar('<', LesserPos))
+		{
+			return String.Mid(LesserPos);
+		}
+		return String;
+	}
+
+	FString TypeConstRef(bool bCastType = false) const
+	{
+		const FString& String = TypeNoRef(bCastType);
+		if (String.StartsWith(TEXT("T"), ESearchCase::CaseSensitive) || String.StartsWith(TEXT("F"), ESearchCase::CaseSensitive))
+		{
+			return FString::Printf(TEXT("const %s&"), *String);
+		}
+		return FString::Printf(TEXT("const %s"), *String);
+	}
+
+	FString TypeRef(bool bCastType = false) const
+	{
+		const FString& String = TypeNoRef(bCastType);
+		return FString::Printf(TEXT("%s&"), *String);
+	}
+
+	FString TypeNoRef(bool bCastType = false) const
+	{
+		const FString& String = TypeOriginal(bCastType);
+		if (String.EndsWith(TEXT("&")))
+		{
+			return String.LeftChop(1);
+		}
+		return String;
+	}
+
+	FString TypeVariableRef(bool bCastType = false) const
+	{
+		return IsConst() ? TypeConstRef(bCastType) : TypeRef(bCastType);
+	}
+
+	FString Variable(bool bCastType = false, bool bCastName = false) const
+	{
+		return FString::Printf(TEXT("%s %s"), *TypeVariableRef(bCastType), *NameOriginal(bCastName));
+	}
+
+	bool IsConst() const
+	{
+		return bConstant || (bInput && !bOutput);
+	}
+
+	bool IsArray() const
+	{
+		return BaseType().Equals(TEXT("TArray"));
+	}
+
+	bool RequiresCast() const
+	{
+		return !CastType.IsEmpty() && !CastName.IsEmpty();
+	}
+};
+
+/**
+ * The FRigVMParameterArray represents the parameters in a notation
+ * of a function marked with RIGVM_METHOD. The parameter array can 
+ * produce a comma separated list of names or parameter declarations.
+ */
+struct FRigVMParameterArray
+{
+public:
+	int32 Num() const { return Parameters.Num(); }
+	const FRigVMParameter& operator[](int32 InIndex) const { return Parameters[InIndex]; }
+	FRigVMParameter& operator[](int32 InIndex) { return Parameters[InIndex]; }
+	TArray<FRigVMParameter>::RangedForConstIteratorType begin() const { return Parameters.begin(); }
+	TArray<FRigVMParameter>::RangedForConstIteratorType end() const { return Parameters.end(); }
+	TArray<FRigVMParameter>::RangedForIteratorType begin() { return Parameters.begin(); }
+	TArray<FRigVMParameter>::RangedForIteratorType end() { return Parameters.end(); }
+
+	int32 Add(const FRigVMParameter& InParameter)
+	{
+		return Parameters.Add(InParameter);
+	}
+
+	FString Names(bool bLeadingSeparator = false, const TCHAR* Separator = TEXT(", "), bool bCastType = false, bool bIncludeEditorOnly = true) const
+	{
+		if (Parameters.Num() == 0)
+		{
+			return FString();
+		}
+		TArray<FString> NameArray;
+		for (const FRigVMParameter& Parameter : Parameters)
+		{
+			if (!bIncludeEditorOnly && Parameter.bEditorOnly)
+			{
+				continue;
+			}
+			NameArray.Add(Parameter.NameOriginal(bCastType));
+		}
+
+		if (NameArray.Num() == 0)
+		{
+			return FString();
+		}
+
+		FString Joined = FString::Join(NameArray, Separator);
+		if (bLeadingSeparator)
+		{
+			return FString::Printf(TEXT("%s%s"), Separator, *Joined);
+		}
+		return Joined;
+	}
+
+	FString Declarations(bool bLeadingSeparator = false, const TCHAR* Separator = TEXT(", "), bool bCastType = false, bool bCastName = false, bool bIncludeEditorOnly = true) const
+	{
+		if (Parameters.Num() == 0)
+		{
+			return FString();
+		}
+		TArray<FString> DeclarationArray;
+		for (const FRigVMParameter& Parameter : Parameters)
+		{
+			if (!bIncludeEditorOnly && Parameter.bEditorOnly)
+			{
+				continue;
+			}
+			DeclarationArray.Add(Parameter.Variable(bCastType, bCastName));
+		}
+
+		if (DeclarationArray.Num() == 0)
+		{
+			return FString();
+		}
+
+		FString Joined = FString::Join(DeclarationArray, Separator);
+		if (bLeadingSeparator)
+		{
+			return FString::Printf(TEXT("%s%s"), Separator, *Joined);
+		}
+		return Joined;
+	}
+
+private:
+	TArray<FRigVMParameter> Parameters;
+};
+
+/**
+ * A single info dataset for a function marked with RIGVM_METHOD.
+ * This struct provides access to its name, the return type and all parameters.
+ */
+struct FRigVMMethodInfo
+{
+	FString ReturnType;
+	FString Name;
+	FRigVMParameterArray Parameters;
+
+	FString ReturnPrefix() const
+	{
+		return (ReturnType.IsEmpty() || (ReturnType == TEXT("void"))) ? TEXT("") : TEXT("return ");
+	}
+};
+
+/**
+ * An info dataset providing access to all functions marked with RIGVM_METHOD
+ * for each struct.
+ */
+struct FRigVMStructInfo
+{
+	FString Name;
+	FRigVMParameterArray Members;
+	TArray<FRigVMMethodInfo> Methods;
+};
+
+typedef TMap<UStruct*, FRigVMStructInfo> FRigVMStructMap;
+
+struct ClassDefinitionRange
+{
+	ClassDefinitionRange(const TCHAR* InStart, const TCHAR* InEnd)
+		: Start(InStart)
+		, End(InEnd)
+		, bHasGeneratedBody(false)
+	{ }
+
+	ClassDefinitionRange()
+		: Start(nullptr)
+		, End(nullptr)
+		, bHasGeneratedBody(false)
+	{ }
+
+	void Validate()
+	{
+		if (End <= Start)
+		{
+			FError::Throwf(TEXT("The class definition range is invalid. Most probably caused by previous parsing error."));
+		}
+	}
+
+	const TCHAR* Start;
+	const TCHAR* End;
+	bool bHasGeneratedBody;
+};
+
+extern TMap<UClass*, ClassDefinitionRange> ClassDefinitionRanges;
+
+#ifndef UHT_DOCUMENTATION_POLICY_DEFAULT
+#define UHT_DOCUMENTATION_POLICY_DEFAULT false
+#endif
+
+struct FDocumentationPolicy
+{
+	bool bClassOrStructCommentRequired = UHT_DOCUMENTATION_POLICY_DEFAULT;
+	bool bFunctionToolTipsRequired = UHT_DOCUMENTATION_POLICY_DEFAULT;
+	bool bMemberToolTipsRequired = UHT_DOCUMENTATION_POLICY_DEFAULT;
+	bool bParameterToolTipsRequired = UHT_DOCUMENTATION_POLICY_DEFAULT;
+	bool bFloatRangesRequired = UHT_DOCUMENTATION_POLICY_DEFAULT;
+};
+
+/////////////////////////////////////////////////////
+// FHeaderParser
+
+//
+// Header parser class.  Extracts metadata from annotated C++ headers and gathers enough
+// information to autogenerate additional headers and other boilerplate code.
+//
+class FHeaderParser : public FBaseParser, public FContextSupplier
+{
+public:
+	// Default version of generated code. Defaults to oldest possible, unless specified otherwise in config.
+	static EGeneratedCodeVersion DefaultGeneratedCodeVersion;
+
+	// Compute the function parameter size and save the return offset
+	static void ComputeFunctionParametersSize(UClass* InClass);
+
+	// Parse all headers for classes that are inside LimitOuter.
+	static ECompilationResult::Type ParseAllHeadersInside(
+		FClasses& ModuleClasses,
+		FFeedbackContext* Warn,
+		UPackage* LimitOuter,
+		const FManifestModule& Module,
+		TArray<class IScriptGeneratorPluginInterface*>& ScriptPlugins
+	);
+
+	// Performs a preliminary parse of the text in the specified buffer, pulling out:
+	//   Class name and parent class name
+	//   Is it an interface
+	//   The list of other classes/interfaces it is dependent on
+	//   
+	//  It also splits the buffer up into:
+	//   ScriptText (text outside of #if CPP and #if DEFAULTS blocks)
+	static void SimplifiedClassParse(const TCHAR* Filename, const TCHAR* Buffer, TArray<FSimplifiedParsingClassInfo>& OutParsedClassArray, TArray<FHeaderProvider>& DependentOn, FStringOutputDevice& ScriptText);
+
+	/**
+	 * Returns True if the given class name includes a valid Unreal prefix and matches up with the given original class Name.
+	 *
+	 * @param InNameToCheck - Name w/ potential prefix to check
+	 * @param OriginalClassName - Name of class w/ no prefix to check against
+	 */
+	static bool ClassNameHasValidPrefix(const FString& InNameToCheck, const FString& OriginalClassName);
+
+	/**
+	 * Tries to convert the header file name to a class name (with 'U' prefix)
+	 *
+	 * @param HeaderFilename Filename.
+	 * @param OutClass The resulting class name (if successfull)
+	 * @return true if the filename was a header filename (.h), false otherwise (in which case OutClassName is unmodified).
+	 */
+	static bool DependentClassNameFromHeader(const TCHAR* HeaderFilename, FString& OutClassName);
+
+	/**
+	 * Transforms CPP-formated string containing default value, to inner formated string
+	 * If it cannot be transformed empty string is returned.
+	 *
+	 * @param Property The property that owns the default value.
+	 * @param CppForm A CPP-formated string.
+	 * @param out InnerForm Inner formated string
+	 * @return true on success, false otherwise.
+	 */
+	static bool DefaultValueStringCppFormatToInnerFormat(const FProperty* Property, const FString& CppForm, FString &InnerForm);
+
+	/**
+	 * Parse Class's annotated headers and optionally its child classes.  Marks the class as CLASS_Parsed.
+	 *
+	 * @param	AllClasses			the class tree containing all classes in the current package
+	 * @param	HeaderParser		the header parser
+	 * @param	SourceFile			Source file info.
+	 *
+	 * @return	Result enumeration.
+	 */
+	static ECompilationResult::Type ParseHeaders(FClasses& AllClasses, FHeaderParser& HeaderParser, FUnrealSourceFile* SourceFile);
+
+protected:
+	friend struct FScriptLocation;
+	friend struct FNativeClassHeaderGenerator;
+
+	// For compiling messages and errors.
+	FFeedbackContext* Warn;
+
+	// Filename currently being parsed
+	FString Filename;
+
+	// Was the first include in the file a validly formed auto-generated header include?
+	bool bSpottedAutogeneratedHeaderInclude;
+
+	// Current nest level, starts at 0.
+	int32 NestLevel;
+
+	// Top nesting level.
+	FNestInfo* TopNest;
+
+	/**
+	 * Gets current nesting scope.
+	 */
+	FScope* GetCurrentScope() const
+	{
+		return TopNest->GetScope();
+	}
+
+	/**
+	 * Gets current file scope.
+	 */
+	FFileScope* GetCurrentFileScope() const
+	{
+		int32 Index = 0;
+		if (!TopNest)
+		{
+			check(!NestLevel);
+			return nullptr;
+		}
+		while (TopNest[Index].NestType != ENestType::GlobalScope)
+		{
+			--Index;
+		}
+
+		return (FFileScope*)TopNest[Index].GetScope();
+	}
+
+	/**
+	 * Gets current source file.
+	 */
+	FUnrealSourceFile* GetCurrentSourceFile() const
+	{
+		return CurrentSourceFile;
+	}
+
+	void SetCurrentSourceFile(FUnrealSourceFile* UnrealSourceFile)
+	{
+		CurrentSourceFile = UnrealSourceFile;
+	}
+
+	/**
+	 * Gets current class scope.
+	 */
+	FStructScope* GetCurrentClassScope() const
+	{
+		check(TopNest->NestType == ENestType::Class || TopNest->NestType == ENestType::Interface || TopNest->NestType == ENestType::NativeInterface);
+
+		return (FStructScope*)TopNest->GetScope();
+	}
+
+	/**
+	 * Tells if parser is currently in a class.
+	 */
+	bool IsInAClass() const
+	{
+		int32 Index = 0;
+		while (TopNest[Index].NestType != ENestType::GlobalScope)
+		{
+			if (TopNest[Index].NestType == ENestType::Class || TopNest->NestType == ENestType::Interface || TopNest->NestType == ENestType::NativeInterface)
+			{
+				return true;
+			}
+
+			--Index;
+		}
+
+		return false;
+	}
+
+	/**
+	 * Gets current class.
+	 */
+	UClass* GetCurrentClass() const
+	{
+		return (UClass*)GetCurrentClassScope()->GetStruct();
+	}
+
+	/**
+	 * Gets current class's metadata.
+	 */
+	FClassMetaData* GetCurrentClassData()
+	{
+		return GScriptHelper.FindClassData(GetCurrentClass());
+	}
+
+	// Information about all nesting levels.
+	FNestInfo Nest[MAX_NEST_LEVELS];
+
+	// enum for complier directives used to build up the directive stack
+	struct ECompilerDirective
+	{
+		enum Type 
+		{
+			// this directive is insignificant and does not change the code generation at all
+			Insignificant			= 0,
+			// this indicates we are in a WITH_EDITOR #if-Block
+			WithEditor				= 1<<0,
+			// this indicates we are in a WITH_EDITORONLY_DATA #if-Block
+			WithEditorOnlyData		= 1<<1,
+		};
+	};
+
+	/** 
+	 * Compiler directive nest in which the parser currently is
+	 * NOTE: compiler directives are combined when more are added onto the stack, so
+	 * checking the only the top of stack is enough to determine in which #if-Block(s) the current code
+	 * is.
+	 *
+	 * ex. Stack.Num() == 1 while entering #if WITH_EDITOR: 
+	 *		CompilerDirectiveStack[1] == CompilerDirectiveStack[0] | ECompilerDirective::WithEditor == 
+	 *		CompilerDirecitveStack[1] == CompilerDirectiveStack.Num()-1 | ECompilerDirective::WithEditor
+	 *
+	 * ex. Stack.Num() == 2 while entering #if WITH_EDITOR: 
+	 *		CompilerDirectiveStack[3] == CompilerDirectiveStack[0] | CompilerDirectiveStack[1] | CompilerDirectiveStack[2] | ECompilerDirective::WithEditor ==
+	 *		CompilerDirecitveStack[3] == CompilerDirectiveStack.Num()-1 | ECompilerDirective::WithEditor
+	 */
+	TArray<uint32> CompilerDirectiveStack;
+
+	// Pushes the Directive specified to the CompilerDirectiveStack according to the rules described above
+	void FORCEINLINE PushCompilerDirective(ECompilerDirective::Type Directive)
+	{
+		CompilerDirectiveStack.Push(CompilerDirectiveStack.Num()>0 ? (CompilerDirectiveStack[CompilerDirectiveStack.Num()-1] | Directive) : Directive);
+	}
+
+	/**
+	 * The starting class flags (i.e. the class flags that were set before the
+	 * CLASS_RecompilerClear mask was applied) for the class currently being compiled
+	 */
+	uint32 PreviousClassFlags;
+
+	// For new-style classes, used to keep track of an unmatched {} pair
+	bool bEncounteredNewStyleClass_UnmatchedBrackets;
+
+	// Indicates that UCLASS/USTRUCT/UINTERFACE has already been parsed in this .h file..
+	bool bHaveSeenUClass;
+
+	// Indicates that a GENERATED_UCLASS_BODY or GENERATED_BODY has been found in the UClass.
+	bool bClassHasGeneratedBody;
+
+	// Indicates that a GENERATED_UINTERFACE_BODY has been found in the UClass.
+	bool bClassHasGeneratedUInterfaceBody;
+
+	// Indicates that a GENERATED_IINTERFACE_BODY has been found in the UClass.
+	bool bClassHasGeneratedIInterfaceBody;
+
+	// public, private, etc at the current parse spot
+	EAccessSpecifier CurrentAccessSpecifier;
+
+	////////////////////////////////////////////////////
+
+	// Special parsed struct names that do not require a prefix
+	static TArray<FString> StructsWithNoPrefix;
+	
+	// Special parsed struct names that have a 'T' prefix
+	static TArray<FString> StructsWithTPrefix;
+
+	// Mapping from 'human-readable' macro substring to # of parameters for delegate declarations
+	// Index 0 is 1 parameter, Index 1 is 2, etc...
+	static TArray<FString> DelegateParameterCountStrings;
+
+	// Types that have been renamed, treat the old deprecated name as the new name for code generation
+	static TMap<FString, FString> TypeRedirectMap;
+
+	// List of all used identifiers for net service function declarations (every function must be unique)
+	TMap<int32, FString> UsedRPCIds;
+	// List of all net service functions with undeclared response functions 
+	TMap<int32, FString> RPCsNeedingHookup;
+
+	// List of all multiplex methods defined on structs
+	static FRigVMStructMap StructRigVMMap;
+
+	// Constructor.
+	explicit FHeaderParser(FFeedbackContext* InWarn, const FManifestModule& InModule);
+
+	virtual ~FHeaderParser()
+	{
+		if ( FScriptLocation::Compiler == this )
+		{
+			FScriptLocation::Compiler = NULL;
+		}
+	}
+
+	// Returns true if the token is a dynamic delegate declaration
+	bool IsValidDelegateDeclaration(const FToken& Token) const;
+
+	// Returns true if the current token is a bitfield type
+	bool IsBitfieldProperty(ELayoutMacroType LayoutMacroType);
+
+	// Parse the parameter list of a function or delegate declaration
+	void ParseParameterList(FClasses& AllClasses, UFunction* Function, bool bExpectCommaBeforeName = false, TMap<FName, FString>* MetaData = NULL);
+
+	// Modify token to fix redirected types if needed
+	void RedirectTypeIdentifier(FToken& Token) const;
+
+public:
+	// Throws if a specifier value wasn't provided
+	static void RequireSpecifierValue(const FPropertySpecifier& Specifier, bool bRequireExactlyOne = false);
+	static FString RequireExactlyOneSpecifierValue(const FPropertySpecifier& Specifier);
+
+	/**
+	* Find a field in the specified context.  Starts with the specified scope, then iterates
+	* through the Outer chain until the field is found.
+	* 
+	* @param	InScope				scope to start searching for the field in 
+	* @param	InIdentifier		name of the field we're searching for
+	* @param	bIncludeParents		whether to allow searching in the scope of a parent struct
+	* @param	FieldClass			class of the field to search for.  used to e.g. search for functions only
+	* @param	Thing				hint text that will be used in the error message if an error is encountered
+	*
+	* @return	a pointer to a UField with a name matching InIdentifier, or NULL if it wasn't found
+	*/
+	static UField* FindField( UStruct* InScope, const TCHAR* InIdentifier, bool bIncludeParents=true, UClass* FieldClass=UField::StaticClass(), const TCHAR* Thing=nullptr );
+	static FField* FindProperty(UStruct* InScope, const TCHAR* InIdentifier, bool bIncludeParents = true, FFieldClass* FieldClass = FField::StaticClass(), const TCHAR* Thing = nullptr);
+
+	// Checks ToValidate to make sure that its associated sparse class data struct, if one exists, is a valid structure to use for storing sparse class data.
+	static void CheckSparseClassData(const UStruct* ToValidate);
+
+protected:
+
+	/**
+	 * Parse rest of the module's source files.
+	 *
+	 * @param	AllClasses The class tree containing all classes in the current package.
+	 * @param	ModulePackage Current package.
+	 * @param	HeaderParser The header parser.
+	 *
+	 * @return	Result enumeration.
+	 */
+	static ECompilationResult::Type ParseRestOfModulesSourceFiles(FClasses& AllClasses, UPackage* ModulePackage, FHeaderParser& HeaderParser);
+
+	//@TODO: Remove this method
+	static void ParseClassName(const TCHAR* Temp, FString& ClassName);
+
+	/**
+	 * @param		Input		An input string, expected to be a script comment.
+	 * @return					The input string, reformatted in such a way as to be appropriate for use as a tooltip.
+	 */
+	static FString FormatCommentForToolTip(const FString& Input);
+
+	/**
+	 * Retrieves parameter comments / tooltips from the function comment
+	 * @param		Input		An input string, expected to be a script comment.
+	 * @return					The map of parameter name to comment per parameter
+	*/
+	static TMap<FName, FString> GetParameterToolTipsFromFunctionComment(const FString& Input);
+	
+	/**
+	 * Begins the process of exporting C++ class declarations for native classes in the specified package
+	 * 
+	 * @param CurrentPackage The package being compiled.
+	 * @param AllClasses The class tree for CurrentPackage.
+	 * @param Module Currently exported module.
+	 */
+	static void ExportNativeHeaders(
+		UPackage* CurrentPackage,
+		FClasses& AllClasses,
+		bool bAllowSaveExportedHeaders,
+		const FManifestModule& Module
+	);
+
+	// FContextSupplier interface.
+	virtual FString GetContext() override;
+	// End of FContextSupplier interface.
+
+	// High-level compiling functions.
+	/**
+	 * Parses given source file.
+	 *
+	 * @param AllClasses The class tree for current package.
+	 * @param SourceFile Source file to parse.
+	 *
+	 * @returns Compilation result enum.
+	 */
+	ECompilationResult::Type ParseHeader(FClasses& AllClasses, FUnrealSourceFile* SourceFile);
+	void CompileDirective(FClasses& AllClasses);
+	void FinalizeScriptExposedFunctions(UClass* Class);
+	UEnum* CompileEnum();
+	UScriptStruct* CompileStructDeclaration(FClasses& AllClasses);
+	bool CompileDeclaration(FClasses& AllClasses, TArray<UDelegateFunction*>& DelegatesToFixup, FToken& Token);
+
+	/** Skip C++ (noexport) declaration. */
+	bool SkipDeclaration(FToken& Token);
+	/** Similar to MatchSymbol() but will return to the exact location as on entry if the symbol was not found. */
+	bool SafeMatchSymbol(const TCHAR Match);
+	void HandleOneInheritedClass(FClasses& AllClasses, UClass* Class, FString&& InterfaceName);
+	FClass* ParseClassNameDeclaration(FClasses& AllClasses, FString& DeclaredClassName, FString& RequiredAPIMacroIfPresent);
+
+	/** The property style of a variable declaration being parsed */
+	struct EPropertyDeclarationStyle
+	{
+		enum Type
+		{
+			None,
+			UPROPERTY
+		};
+	};
+
+	/**
+	 * Resets current class data back to its defaults.
+	 */
+	void ResetClassData();
+
+	/**
+	 * Create new function object based on given info structure.
+	 */
+	UFunction* CreateFunction(const FFuncInfo &FuncInfo) const;
+
+	/**
+	 * Create new delegate function object based on given info structure.
+	 */
+	template<typename T>
+	UDelegateFunction* CreateDelegateFunction(const FFuncInfo &FuncInfo) const;	
+
+	UClass* CompileClassDeclaration(FClasses& AllClasses);
+	UDelegateFunction* CompileDelegateDeclaration(FClasses& AllClasses, const TCHAR* DelegateIdentifier, EDelegateSpecifierAction::Type SpecifierAction = EDelegateSpecifierAction::DontParse);
+	void CompileFunctionDeclaration(FClasses& AllClasses);
+	void CompileVariableDeclaration (FClasses& AllClasses, UStruct* Struct);
+	void CompileInterfaceDeclaration(FClasses& AllClasses);
+	void CompileRigVMMethodDeclaration(FClasses& AllClasses, UStruct* Struct);
+	void ParseRigVMMethodParameters(UStruct* Struct);
+
+	FClass* ParseInterfaceNameDeclaration(FClasses& AllClasses, FString& DeclaredInterfaceName, FString& RequiredAPIMacroIfPresent);
+	bool TryParseIInterfaceClass(FClasses& AllClasses);
+
+	bool CompileStatement(FClasses& AllClasses, TArray<UDelegateFunction*>& DelegatesToFixup);
+
+	// Checks to see if a particular kind of command is allowed on this nesting level.
+	bool IsAllowedInThisNesting(ENestAllowFlags AllowFlags);
+	
+	// Make sure that a particular kind of command is allowed on this nesting level.
+	// If it's not, issues a compiler error referring to the token and the current
+	// nesting level.
+	void CheckAllow(const TCHAR* Thing, ENestAllowFlags AllowFlags);
+
+	UStruct* GetSuperScope( UStruct* CurrentScope, const FName& SearchName );	
+
+	void SkipStatements( int32 SubCount, const TCHAR* ErrorTag );
+
+	/**
+	 * Parses a variable or return value declaration and determines the variable type and property flags.
+	 *
+	 * @param   AllClasses                the class tree for CurrentPackage
+	 * @param   Scope                     struct to create the property in
+	 * @param   VarProperty               will be filled in with type and property flag data for the property declaration that was parsed
+	 * @param   Disallow                  contains a mask of variable modifiers that are disallowed in this context
+	 * @param   OuterPropertyType         only specified when compiling the inner properties for arrays or maps.  corresponds to the FToken for the outer property declaration.
+	 * @param   PropertyDeclarationStyle  if the variable is defined with a UPROPERTY
+	 * @param   VariableCategory          what kind of variable is being parsed
+	 * @param   ParsedVarIndexRange       The source text [Start, End) index range for the parsed type.
+	 */
+	void GetVarType(
+		FClasses&                       AllClasses,
+		FScope*							Scope,
+		FPropertyBase&                  VarProperty,
+		EPropertyFlags                  Disallow,
+		const FToken*                   OuterPropertyType,
+		EPropertyDeclarationStyle::Type PropertyDeclarationStyle,
+		EVariableCategory::Type         VariableCategory,
+		FIndexRange*                    ParsedVarIndexRange = nullptr,
+		ELayoutMacroType*               OutLayoutMacroType = nullptr);
+
+	/**
+	 * Parses a variable name declaration and creates a new FProperty object.
+	 *
+	 * @param	Scope				struct to create the property in
+	 * @param	VarProperty			type and propertyflag info for the new property (inout)
+	 * @param   VariableCategory	what kind of variable is being created
+	 *
+	 * @return	a pointer to the new FProperty if successful, or NULL if there was no property to parse
+	 */
+	FProperty* GetVarNameAndDim(
+		UStruct* Struct,
+		FToken& VarProperty,
+		EVariableCategory::Type VariableCategory,
+		ELayoutMacroType LayoutMacroType = ELayoutMacroType::None);
+	
+	/**
+	 * Returns whether the specified class can be referenced from the class currently being compiled.
+	 *
+	 * @param	Scope		The scope we are currently parsing.
+	 * @param	CheckClass	The class we want to reference.
+	 *
+	 * @return	true if the specified class is an intrinsic type or if the class has successfully been parsed
+	 */
+	bool AllowReferenceToClass(UStruct* Scope, UClass* CheckClass) const;
+
+	/**
+	 * @return	true if Scope has FProperty objects in its list of fields
+	 */
+	static bool HasMemberProperties( const UStruct* Scope );
+
+	/**
+	 * Parses optional metadata text.
+	 *
+	 * @param	MetaData	the metadata map to store parsed metadata in
+	 * @param	FieldName	the field being parsed (used for logging)
+	 *
+	 * @return	true if metadata was specified
+	 */
+	void ParseFieldMetaData(TMap<FName, FString>& MetaData, const TCHAR* FieldName);
+
+	/**
+	 * Formats the current comment, if any, and adds it to the metadata as a tooltip.
+	 *
+	 * @param	MetaData	the metadata map to store the tooltip in
+	 */
+	void AddFormattedPrevCommentAsTooltipMetaData(TMap<FName, FString>& MetaData);
+
+	/** 
+	 * Tries to parse the token as an access protection specifier (public:, protected:, or private:)
+	 *
+	 * @return	EAccessSpecifier this is, or zero if it is none
+	 */
+	EAccessSpecifier ParseAccessProtectionSpecifier(const FToken& Token);
+
+	const TCHAR* NestTypeName( ENestType NestType );
+
+	FClass* GetQualifiedClass(const FClasses& AllClasses, const TCHAR* Thing);
+
+	/**
+	 * Increase the nesting level, setting the new top nesting level to
+	 * the one specified.  If pushing a function or state and it overrides a similar
+	 * thing declared on a lower nesting level, verifies that the override is legal.
+	 *
+	 * @param	NestType	the new nesting type
+	 * @param	InNode		@todo
+	 */
+	void PushNest(ENestType NestType, UStruct* InNode, FUnrealSourceFile* SourceFile = nullptr);
+	void PopNest(ENestType NestType, const TCHAR* Descr);
+
+	/**
+	 * Tasks that need to be done after popping function declaration
+	 * from parsing stack.
+	 *
+	 * @param AllClasses The class tree for current package.
+	 * @param PoppedFunction Function that have just been popped.
+	 */
+	void PostPopFunctionDeclaration(FClasses& AllClasses, UFunction* PoppedFunction);
+
+	/**
+	 * Tasks that need to be done after popping interface definition
+	 * from parsing stack.
+	 *
+	 * @param AllClasses The class tree for current package.
+	 * @param CurrentInterface Interface that have just been popped.
+	 */
+	void PostPopNestInterface(FClasses& AllClasses, UClass* CurrentInterface);
+
+	/**
+	 * Tasks that need to be done after popping class definition
+	 * from parsing stack.
+	 *
+	 * @param CurrentClass Class that have just been popped.
+	 */
+	void PostPopNestClass(UClass* CurrentClass);
+
+	/**
+	 * Binds all delegate properties declared in ValidationScope the delegate functions specified in the variable declaration, verifying that the function is a valid delegate
+	 * within the current scope.  This must be done once the entire class has been parsed because instance delegate properties must be declared before the delegate declaration itself.
+	 *
+	 * @todo: this function will no longer be required once the post-parse fixup phase is added (TTPRO #13256)
+	 *
+	 * @param	AllClasses			the class tree for CurrentPackage
+	 * @param	Struct				the struct to validate delegate properties for
+	 * @param	Scope				the current scope
+	 * @param	DelegateCache		cached map of delegates that have already been found; used for faster lookup.
+	 */
+	void FixupDelegateProperties(FClasses& AllClasses, UStruct* ValidationScope, FScope& Scope, TMap<FName, UFunction*>& DelegateCache);
+
+	// Retry functions.
+	void InitScriptLocation( FScriptLocation& Retry );
+	void ReturnToLocation( const FScriptLocation& Retry, bool Binary=1, bool Text=1 );
+
+	/**
+	 * If the property has already been seen during compilation, then return add. If not,
+	 * then return replace so that INI files don't mess with header exporting
+	 *
+	 * @param PropertyName the string token for the property
+	 *
+	 * @return FNAME_Replace_Not_Safe_For_Threading or FNAME_Add
+	 */
+	EFindName GetFindFlagForPropertyName(const TCHAR* PropertyName);
+
+	static void ValidatePropertyIsDeprecatedIfNecessary(const FPropertyBase& VarProperty, const FToken* OuterPropertyType);
+
+private:
+	// Source file currently parsed by UHT.
+	FUnrealSourceFile* CurrentSourceFile;
+
+	// Module currently parsed by UHT.
+	const FManifestModule* CurrentlyParsedModule;
+
+	// True if the module currently being parsed is part of the engine, as opposed to being part of a game
+	bool bIsCurrentModulePartOfEngine;
+
+	/**
+	 * Tries to match constructor parameter list. Assumes that constructor
+	 * name is already matched.
+	 *
+	 * If fails it reverts all parsing done.
+	 *
+	 * @param Token Token to start parsing from.
+	 *
+	 * @returns True if matched. False otherwise.
+	 */
+	bool TryToMatchConstructorParameterList(FToken Token);
+
+	// Parses possible version declaration in generated code, e.g. GENERATED_BODY(<some_version>).
+	void CompileVersionDeclaration(UStruct* Struct);
+
+	// Verifies that all specified class's UProperties with function associations have valid targets
+	void VerifyPropertyMarkups( UClass* TargetClass );
+
+	// Verifies the target function meets the criteria for a blueprint property getter
+	void VerifyBlueprintPropertyGetter(FProperty* Property, UFunction* TargetFunction);
+
+	// Verifies the target function meets the criteria for a blueprint property setter
+	void VerifyBlueprintPropertySetter(FProperty* Property, UFunction* TargetFunction);
+
+	// Verifies the target function meets the criteria for a replication notify callback
+	void VerifyRepNotifyCallback(FProperty* Property, UFunction* TargetFunction);
+
+	// Constructs the policy from a string
+	static FDocumentationPolicy GetDocumentationPolicyFromName(const FString& PolicyName);
+
+	// Constructs the policy for documentation checks for a given struct
+	static FDocumentationPolicy GetDocumentationPolicyForStruct(UStruct* Struct);
+
+	// Property types to provide UI Min and Max ranges
+	static TArray<FString> PropertyCPPTypesRequiringUIRanges;
+
+	// Returns true if a given CPP types required ui checking
+	static bool DoesCPPTypeRequireDocumentation(const FString& CPPType);
+
+	// Validates the documentation for a given enum
+	void CheckDocumentationPolicyForEnum(UEnum* Enum, const TMap<FName, FString>& MetaData, const TArray<TMap<FName, FString>>& Entries);
+
+	// Validates the documentation for a given struct
+	void CheckDocumentationPolicyForStruct(UStruct* Struct);
+
+	// Validates the documentation for a given method
+	void CheckDocumentationPolicyForFunc(UClass* Class, UFunction* Func);
+
+	// Checks if a valid range has been found on the provided metadata
+	bool CheckUIMinMaxRangeFromMetaData(const FString& UIMin, const FString& UIMax);
+
+	// Names that cannot be used enums, UStructs, or UClasses
+	static TArray<FString> ReservedTypeNames;
+
+public:
+	/**
+	* Checks if the given token uses one of the reserved type names.
+	*
+	* @param  TypeName		String of the type to check (For UObject/UClass, use the stripped name)
+	* @return True if the TypeName is a reserved name
+	*/
+	static bool IsReservedTypeName(const FString& TypeName);
+
+	/**
+	* Checks if the given token uses one of the reserved type names.
+	*
+	* @param  Token		The token to check
+	* @return True if the Token is using a reserved name
+	*/
+	static bool IsReservedTypeName(const FToken& Token);
+};
+
+/////////////////////////////////////////////////////
+// FHeaderPreParser
+
+class FHeaderPreParser : public FBaseParser
+{
+public:
+	FHeaderPreParser()
+	{
+	}
+
+	void ParseClassDeclaration(
+		const TCHAR* Filename,
+		const TCHAR* InputText,
+		int32 InLineNumber,
+		const TCHAR*
+		StartingMatchID,
+		FName& out_StrippedClassName,
+		FString& out_ClassName,
+		FString& out_BaseClassName,
+		TArray<FHeaderProvider>& out_ClassNames,
+		const TArray<FSimplifiedParsingClassInfo>& ParsedClassArray
+	);
+};
diff --git a/Engine/Source/Programs/UHTLite/Private/HeaderProvider.cpp b/Engine/Source/Programs/UHTLite/Private/HeaderProvider.cpp
new file mode 100644
--- /dev/null
+++ b/Engine/Source/Programs/UHTLite/Private/HeaderProvider.cpp
@@ -0,0 +1,52 @@
+// Copyright Epic Games, Inc. All Rights Reserved.
+
+#include "HeaderProvider.h"
+#include "UnrealHeaderTool.h"
+#include "UnrealTypeDefinitionInfo.h"
+#include "ClassMaps.h"
+
+FHeaderProvider::FHeaderProvider(EHeaderProviderSourceType InType, FString&& InId)//, bool bInAutoInclude/* = false*/)
+	: Type(InType)
+	, Id(MoveTemp(InId))
+	, Cache(nullptr)
+{
+
+}
+
+FUnrealSourceFile* FHeaderProvider::Resolve()
+{
+	if (Type != EHeaderProviderSourceType::Resolved)
+	{
+		if (Type == EHeaderProviderSourceType::ClassName)
+		{
+			FName IdName(*Id, FNAME_Find);
+			if (TSharedRef<FUnrealTypeDefinitionInfo>* Source = GTypeDefinitionInfoMap.FindByName(IdName))
+			{
+				Cache = &(*Source)->GetUnrealSourceFile();
+			}
+		}
+		else if (const TSharedRef<FUnrealSourceFile>* Source = GUnrealSourceFilesMap.Find(Id))
+		{
+			Cache = &Source->Get();
+		}
+
+		Type = EHeaderProviderSourceType::Resolved;
+	}
+
+	return Cache;
+}
+
+FString FHeaderProvider::ToString() const
+{
+	return FString::Printf(TEXT("%s %s"), Type == EHeaderProviderSourceType::ClassName ? TEXT("class") : TEXT("file"), *Id);
+}
+
+const FString& FHeaderProvider::GetId() const
+{
+	return Id;
+}
+
+bool operator==(const FHeaderProvider& A, const FHeaderProvider& B)
+{
+	return A.Type == B.Type && A.Id == B.Id;
+}
diff --git a/Engine/Source/Programs/UHTLite/Private/HeaderProvider.h b/Engine/Source/Programs/UHTLite/Private/HeaderProvider.h
new file mode 100644
--- /dev/null
+++ b/Engine/Source/Programs/UHTLite/Private/HeaderProvider.h
@@ -0,0 +1,35 @@
+// Copyright Epic Games, Inc. All Rights Reserved.
+
+#pragma once
+
+#include "CoreMinimal.h"
+#include "Containers/UnrealString.h"
+
+class FUnrealSourceFile;
+
+enum class EHeaderProviderSourceType
+{
+	ClassName,
+	FileName,
+	Resolved
+};
+
+class FHeaderProvider
+{
+	friend bool operator==(const FHeaderProvider& A, const FHeaderProvider& B);
+public:
+	FHeaderProvider(EHeaderProviderSourceType Type, FString&& Id);
+
+	FUnrealSourceFile* Resolve();
+
+	FString ToString() const;
+
+	const FString& GetId() const;
+
+private:
+	EHeaderProviderSourceType Type;
+	FString Id;
+	FUnrealSourceFile* Cache;
+};
+
+bool operator==(const FHeaderProvider& A, const FHeaderProvider& B);
diff --git a/Engine/Source/Programs/UHTLite/Private/IScriptGeneratorPluginInterface.cpp b/Engine/Source/Programs/UHTLite/Private/IScriptGeneratorPluginInterface.cpp
new file mode 100644
--- /dev/null
+++ b/Engine/Source/Programs/UHTLite/Private/IScriptGeneratorPluginInterface.cpp
@@ -0,0 +1,74 @@
+// Copyright Epic Games, Inc. All Rights Reserved.
+
+#include "IScriptGeneratorPluginInterface.h"
+#include "UnrealHeaderTool.h"
+#include "UObject/ErrorException.h"
+#include "Algo/FindSortedStringCaseInsensitive.h"
+
+EBuildModuleType::Type EBuildModuleType::Parse(const TCHAR* Value)
+{
+	static const TCHAR* AlphabetizedTypes[] = {
+		TEXT("EngineDeveloper"),
+		TEXT("EngineEditor"),
+		TEXT("EngineRuntime"),
+		TEXT("EngineThirdParty"),
+		TEXT("EngineUncooked"),
+		TEXT("GameDeveloper"),
+		TEXT("GameEditor"),
+		TEXT("GameRuntime"),
+		TEXT("GameThirdParty"),
+		TEXT("GameUncooked"),
+		TEXT("Program")
+	};
+
+	int32 TypeIndex = Algo::FindSortedStringCaseInsensitive(Value, AlphabetizedTypes);
+	if (TypeIndex < 0)
+	{
+		FError::Throwf(TEXT("Unrecognized EBuildModuleType name: %s"), Value);
+	}
+
+	static EBuildModuleType::Type AlphabetizedValues[] = {
+		EngineDeveloper,
+		EngineEditor,
+		EngineRuntime,
+		EngineThirdParty,
+		EngineUncooked,
+		GameDeveloper,
+		GameEditor,
+		GameRuntime,
+		GameThirdParty,
+		GameUncooked,
+		Program
+	};
+
+	return AlphabetizedValues[TypeIndex];
+}
+
+EPackageOverrideType::Type EPackageOverrideType::Parse(const TCHAR* Value)
+{
+	static const TCHAR* AlphabetizedTypes[] = {
+		TEXT("EditorOnly"),
+		TEXT("EngineDeveloper"),
+		TEXT("EngineUncookedOnly"),
+		TEXT("GameDeveloper"),
+		TEXT("GameUncookedOnly"),
+		TEXT("None"),
+	};
+
+	int32 TypeIndex = Algo::FindSortedStringCaseInsensitive(Value, AlphabetizedTypes);
+	if (TypeIndex < 0)
+	{
+		FError::Throwf(TEXT("Unrecognized EPackageOverrideType name: %s"), Value);
+	}
+
+	static EPackageOverrideType::Type AlphabetizedValues[] = {
+		EditorOnly,
+		EngineDeveloper,
+		EngineUncookedOnly,
+		GameDeveloper,
+		GameUncookedOnly,
+		None
+	};
+
+	return AlphabetizedValues[TypeIndex];
+}
\ No newline at end of file
diff --git a/Engine/Source/Programs/UHTLite/Private/Manifest.cpp b/Engine/Source/Programs/UHTLite/Private/Manifest.cpp
new file mode 100644
--- /dev/null
+++ b/Engine/Source/Programs/UHTLite/Private/Manifest.cpp
@@ -0,0 +1,255 @@
+// Copyright Epic Games, Inc. All Rights Reserved.
+
+#include "Manifest.h"
+#include "UnrealHeaderTool.h"
+#include "Misc/DateTime.h"
+#include "Logging/LogMacros.h"
+#include "HAL/FileManager.h"
+#include "Misc/FileHelper.h"
+#include "Misc/Paths.h"
+#include "UObject/ErrorException.h"
+#include "Misc/PackageName.h"
+#include "UnrealHeaderToolGlobals.h"
+#include "Serialization/JsonTypes.h"
+#include "Serialization/JsonSerializer.h"
+
+namespace
+{
+	template <typename T> struct TJsonFieldType;
+	template <>           struct TJsonFieldType<double>                         { static const EJson Value = EJson::Number;  };
+	template <>           struct TJsonFieldType<FString>                        { static const EJson Value = EJson::String;  };
+	template <>           struct TJsonFieldType<bool>                           { static const EJson Value = EJson::Boolean; };
+	template <>           struct TJsonFieldType<TArray<TSharedPtr<FJsonValue>>> { static const EJson Value = EJson::Array;   };
+	template <>           struct TJsonFieldType<TSharedPtr<FJsonObject>>        { static const EJson Value = EJson::Object;  };
+
+	template <typename T>
+	void GetJsonValue(T& OutVal, const TSharedPtr<FJsonValue>& JsonValue, const TCHAR* Outer)
+	{
+		if (JsonValue->Type != TJsonFieldType<T>::Value)
+		{
+			FError::Throwf(TEXT("'%s' is the wrong type"), Outer);
+		}
+
+		JsonValue->AsArgumentType(OutVal);
+	}
+
+	template <typename T>
+	void GetJsonFieldValue(T& OutVal, const TSharedPtr<FJsonObject>& JsonObject, const TCHAR* FieldName, const TCHAR* Outer)
+	{
+		TSharedPtr<FJsonValue>* JsonValue = JsonObject->Values.Find(FieldName);
+
+		if (!JsonValue)
+		{
+			FError::Throwf(TEXT("Unable to find field '%s' in '%s'"), FieldName, Outer);
+		}
+
+		if ((*JsonValue)->Type != TJsonFieldType<T>::Value)
+		{
+			FError::Throwf(TEXT("Field '%s' in '%s' is the wrong type"), Outer);
+		}
+
+		(*JsonValue)->AsArgumentType(OutVal);
+	}
+
+	void ProcessHeaderArray(FString* OutStringArray, const TArray<TSharedPtr<FJsonValue>>& InJsonArray, const TCHAR* Outer)
+	{
+		for (int32 Index = 0, Count = InJsonArray.Num(); Index != Count; ++Index)
+		{
+			GetJsonValue(*OutStringArray++, InJsonArray[Index], *FString::Printf(TEXT("%s[%d]"), Outer, Index));
+		}
+	}
+}
+
+
+FManifest FManifest::LoadFromFile(const FString& Filename)
+{
+	FManifest Result;
+	FString FilenamePath = FPaths::GetPath(Filename);
+	FString Json;
+
+	if (!FFileHelper::LoadFileToString(Json, *Filename))
+	{
+		FError::Throwf(TEXT("Unable to load manifest: %s"), *Filename);
+	}
+
+	TSharedPtr<FJsonObject> RootObject = TSharedPtr<FJsonObject>();
+	TSharedRef<TJsonReader<TCHAR>> Reader = TJsonReaderFactory<TCHAR>::Create(Json);
+
+	if (!FJsonSerializer::Deserialize(Reader, RootObject))
+	{
+		FError::Throwf(TEXT("Manifest is malformed: %s"), *Filename);
+	}
+
+	TArray<TSharedPtr<FJsonValue>> ModulesArray;
+
+	GetJsonFieldValue(Result.IsGameTarget,             RootObject, TEXT("IsGameTarget"),             TEXT("{manifest root}"));
+	GetJsonFieldValue(Result.RootLocalPath,            RootObject, TEXT("RootLocalPath"),            TEXT("{manifest root}"));
+	Result.RootBuildPath = Result.RootLocalPath + FPlatformMisc::GetDefaultPathSeparator();
+	GetJsonFieldValue(Result.TargetName,               RootObject, TEXT("TargetName"),               TEXT("{manifest root}"));
+	GetJsonFieldValue(Result.ExternalDependenciesFile, RootObject, TEXT("ExternalDependenciesFile"), TEXT("{manifest root}"));
+	GetJsonFieldValue(ModulesArray,                    RootObject, TEXT("Modules"),                  TEXT("{manifest root}"));
+
+	UE_LOG(LogCompile, Log, TEXT("Loaded manifest: %s"),       *Filename);
+	UE_LOG(LogCompile, Log, TEXT("Manifest.IsGameTarget=%s"),  Result.IsGameTarget ? TEXT("True") : TEXT("False"));
+	UE_LOG(LogCompile, Log, TEXT("Manifest.RootLocalPath=%s"), *Result.RootLocalPath);
+	UE_LOG(LogCompile, Log, TEXT("Manifest.RootBuildPath=%s"), *Result.RootBuildPath);
+	UE_LOG(LogCompile, Log, TEXT("Manifest.TargetName=%s"),    *Result.TargetName);
+	UE_LOG(LogCompile, Log, TEXT("Manifest.Modules=%d"),       ModulesArray.Num());
+
+	Result.RootLocalPath = FPaths::ConvertRelativePathToFull(FilenamePath, Result.RootLocalPath);
+	Result.RootBuildPath = FPaths::ConvertRelativePathToFull(FilenamePath, Result.RootBuildPath);
+
+	// Ensure directories end with a slash, because this aids their use with FPaths::MakePathRelativeTo.
+	if (!Result.RootLocalPath.EndsWith(TEXT("/")))
+	{
+		Result.RootLocalPath += TEXT("/");
+	}
+
+	if (!Result.RootBuildPath.EndsWith(TEXT("/")))
+	{
+		Result.RootBuildPath += TEXT("/");
+	}
+
+	int32 ModuleIndex = 0;
+
+	for (const TSharedPtr<FJsonValue>& Module : ModulesArray)
+	{
+		const TSharedPtr<FJsonObject>& ModuleObj = Module->AsObject();
+
+		TArray<TSharedPtr<FJsonValue>> ClassesHeaders;
+		TArray<TSharedPtr<FJsonValue>> PublicHeaders;
+		TArray<TSharedPtr<FJsonValue>> PrivateHeaders;
+
+		Result.Modules.AddZeroed();
+		FManifestModule& KnownModule = Result.Modules.Last();
+
+		FString Outer = FString::Printf(TEXT("Modules[%d]"), ModuleIndex);
+		FString GeneratedCodeVersionString;
+		GetJsonFieldValue(KnownModule.Name,                      ModuleObj, TEXT("Name"),                     *Outer);
+		GetJsonFieldValue(KnownModule.BaseDirectory,             ModuleObj, TEXT("BaseDirectory"),            *Outer);
+		GetJsonFieldValue(KnownModule.IncludeBase,               ModuleObj, TEXT("IncludeBase"),              *Outer);
+		GetJsonFieldValue(KnownModule.GeneratedIncludeDirectory, ModuleObj, TEXT("OutputDirectory"),          *Outer);
+		GetJsonFieldValue(KnownModule.SaveExportedHeaders,       ModuleObj, TEXT("SaveExportedHeaders"),      *Outer);
+		GetJsonFieldValue(ClassesHeaders,                        ModuleObj, TEXT("ClassesHeaders"),           *Outer);
+		GetJsonFieldValue(PublicHeaders,                         ModuleObj, TEXT("PublicHeaders"),            *Outer);
+		GetJsonFieldValue(PrivateHeaders,                        ModuleObj, TEXT("PrivateHeaders"),           *Outer);
+		GetJsonFieldValue(KnownModule.GeneratedCPPFilenameBase,  ModuleObj, TEXT("GeneratedCPPFilenameBase"), *Outer);
+		GetJsonFieldValue(GeneratedCodeVersionString, ModuleObj, TEXT("UHTGeneratedCodeVersion"), *Outer);
+		KnownModule.GeneratedCodeVersion = ToGeneratedCodeVersion(GeneratedCodeVersionString);
+
+		FString ModuleTypeText;
+		GetJsonFieldValue(ModuleTypeText, ModuleObj, TEXT("ModuleType"), *Outer);
+		KnownModule.ModuleType = EBuildModuleType::Parse(*ModuleTypeText);
+
+		FString OverrideModuleTypeText;
+		GetJsonFieldValue(OverrideModuleTypeText, ModuleObj, TEXT("OverrideModuleType"), *Outer);
+		KnownModule.OverrideModuleType = EPackageOverrideType::Parse(*OverrideModuleTypeText);
+
+		KnownModule.LongPackageName = FPackageName::ConvertToLongScriptPackageName(*KnownModule.Name);
+
+		// Convert relative paths
+		KnownModule.BaseDirectory             = FPaths::ConvertRelativePathToFull(FilenamePath, KnownModule.BaseDirectory);
+		KnownModule.IncludeBase               = FPaths::ConvertRelativePathToFull(FilenamePath, KnownModule.IncludeBase);
+		KnownModule.GeneratedIncludeDirectory = FPaths::ConvertRelativePathToFull(FilenamePath, KnownModule.GeneratedIncludeDirectory);
+		KnownModule.GeneratedCPPFilenameBase  = FPaths::ConvertRelativePathToFull(FilenamePath, KnownModule.GeneratedCPPFilenameBase);
+
+		// Ensure directories end with a slash, because this aids their use with FPaths::MakePathRelativeTo.
+		if (!KnownModule.BaseDirectory            .EndsWith(TEXT("/"))) { KnownModule.BaseDirectory            .AppendChar(TEXT('/')); }
+		if (!KnownModule.IncludeBase              .EndsWith(TEXT("/"))) { KnownModule.IncludeBase              .AppendChar(TEXT('/')); }
+		if (!KnownModule.GeneratedIncludeDirectory.EndsWith(TEXT("/"))) { KnownModule.GeneratedIncludeDirectory.AppendChar(TEXT('/')); }
+
+		KnownModule.PublicUObjectClassesHeaders.AddZeroed(ClassesHeaders.Num());
+		KnownModule.PublicUObjectHeaders       .AddZeroed(PublicHeaders .Num());
+		KnownModule.PrivateUObjectHeaders      .AddZeroed(PrivateHeaders.Num());
+
+		ProcessHeaderArray(KnownModule.PublicUObjectClassesHeaders.GetData(), ClassesHeaders, *(Outer + TEXT(".ClassHeaders")));
+		ProcessHeaderArray(KnownModule.PublicUObjectHeaders       .GetData(), PublicHeaders , *(Outer + TEXT(".PublicHeaders" )));
+		ProcessHeaderArray(KnownModule.PrivateUObjectHeaders      .GetData(), PrivateHeaders, *(Outer + TEXT(".PrivateHeaders")));
+
+		// Sort the headers alphabetically.  This is just to add determinism to the compilation dependency order, since we currently
+		// don't rely on explicit includes (but we do support 'dependson')
+		// @todo uht: Ideally, we should sort these by sensical order before passing them in -- or better yet, follow include statements ourselves in here.
+		KnownModule.PublicUObjectClassesHeaders.Sort();
+		KnownModule.PublicUObjectHeaders       .Sort();
+		KnownModule.PrivateUObjectHeaders      .Sort();
+
+		UE_LOG(LogCompile, Log, TEXT("  %s"), *KnownModule.Name);
+		UE_LOG(LogCompile, Log, TEXT("  .BaseDirectory=%s"), *KnownModule.BaseDirectory);
+		UE_LOG(LogCompile, Log, TEXT("  .IncludeBase=%s"), *KnownModule.IncludeBase);
+		UE_LOG(LogCompile, Log, TEXT("  .GeneratedIncludeDirectory=%s"), *KnownModule.GeneratedIncludeDirectory);
+		UE_LOG(LogCompile, Log, TEXT("  .SaveExportedHeaders=%s"), KnownModule.SaveExportedHeaders ? TEXT("True") : TEXT("False"));
+		UE_LOG(LogCompile, Log, TEXT("  .GeneratedCPPFilenameBase=%s"), *KnownModule.GeneratedCPPFilenameBase);
+		UE_LOG(LogCompile, Log, TEXT("  .ModuleType=%s"), *ModuleTypeText);
+
+		++ModuleIndex;
+	}
+
+	return Result;
+}
+
+bool FManifestModule::NeedsRegeneration() const
+{
+	if (ShouldForceRegeneration())
+	{
+		return true;
+	}
+	FString Timestamp;
+	TCHAR TimestampText[] = TEXT("Timestamp");
+	Timestamp.Empty(GeneratedIncludeDirectory.Len() + UE_ARRAY_COUNT(TimestampText));
+	Timestamp += GeneratedIncludeDirectory;
+	Timestamp += TEXT("Timestamp");
+
+	if (!FPaths::FileExists(Timestamp))
+	{
+		// No timestamp, must regenerate.
+		return true;
+	}
+
+	FDateTime TimestampFileLastModify = IFileManager::Get().GetTimeStamp(*Timestamp);
+
+	for (const FString& Header : PublicUObjectClassesHeaders)
+	{
+		if (IFileManager::Get().GetTimeStamp(*Header) > TimestampFileLastModify)
+		{
+			UE_LOG(LogCompile, Log, TEXT("File %s is newer than last timestamp. Regenerating reflection data for module %s."), *Header, *Name);
+			return true;
+		}
+	}
+
+	for (const FString& Header : PublicUObjectHeaders)
+	{
+		if (IFileManager::Get().GetTimeStamp(*Header) > TimestampFileLastModify)
+		{
+			UE_LOG(LogCompile, Log, TEXT("File %s is newer than last timestamp. Regenerating reflection data for module %s."), *Header, *Name);
+			return true;
+		}
+	}
+
+	for (const FString& Header : PrivateUObjectHeaders)
+	{
+		if (IFileManager::Get().GetTimeStamp(*Header) > TimestampFileLastModify)
+		{
+			UE_LOG(LogCompile, Log, TEXT("File %s is newer than last timestamp. Regenerating reflection data for module %s."), *Header, *Name);
+			return true;
+		}
+	}
+
+	// No header is newer than timestamp, no need to regenerate.
+	return false;
+}
+
+bool FManifestModule::IsCompatibleWith(const FManifestModule& ManifestModule)
+{
+	return ModuleType == ManifestModule.ModuleType
+		&& SaveExportedHeaders == ManifestModule.SaveExportedHeaders
+		&& GeneratedCodeVersion == ManifestModule.GeneratedCodeVersion
+		&& Name == ManifestModule.Name
+		&& LongPackageName == ManifestModule.LongPackageName
+		&& BaseDirectory == ManifestModule.BaseDirectory
+		&& IncludeBase == ManifestModule.IncludeBase
+		&& GeneratedIncludeDirectory == ManifestModule.GeneratedIncludeDirectory
+		&& PublicUObjectClassesHeaders == ManifestModule.PublicUObjectClassesHeaders
+		&& PublicUObjectHeaders == ManifestModule.PublicUObjectHeaders
+		&& PrivateUObjectHeaders == ManifestModule.PrivateUObjectHeaders
+		&& GeneratedCPPFilenameBase == ManifestModule.GeneratedCPPFilenameBase;
+}
diff --git a/Engine/Source/Programs/UHTLite/Private/Manifest.h b/Engine/Source/Programs/UHTLite/Private/Manifest.h
new file mode 100644
--- /dev/null
+++ b/Engine/Source/Programs/UHTLite/Private/Manifest.h
@@ -0,0 +1,122 @@
+// Copyright Epic Games, Inc. All Rights Reserved.
+
+#pragma once
+
+#include "CoreMinimal.h"
+#include "GeneratedCodeVersion.h"
+#include "IScriptGeneratorPluginInterface.h"
+
+class FArchive;
+
+struct FManifestModule
+{
+	/** The name of the module */
+	FString Name;
+
+	/** Module type */
+	EBuildModuleType::Type ModuleType;
+
+	/** Overridden package settings to add additional flags that can help with organization */
+	EPackageOverrideType::Type OverrideModuleType;
+
+	/** Long package name for this module's UObject class */
+	FString LongPackageName;
+
+	/** Base directory of this module on disk */
+	FString BaseDirectory;
+
+	/** The directory to which #includes from this module should be relative */
+	FString IncludeBase;
+
+	/** Directory where generated include files should go */
+	FString GeneratedIncludeDirectory;
+
+	/** List of C++ public 'Classes' header files with UObjects in them (legacy) */
+	TArray<FString> PublicUObjectClassesHeaders;
+
+	/** List of C++ public header files with UObjects in them */
+	TArray<FString> PublicUObjectHeaders;
+
+	/** List of C++ private header files with UObjects in them */
+	TArray<FString> PrivateUObjectHeaders;
+
+	/** Base (i.e. extensionless) path+filename of where to write out the module's .generated.* files */
+	FString GeneratedCPPFilenameBase;
+
+	/** Whether or not to write out headers that have changed */
+	bool SaveExportedHeaders;
+
+	/** Version of generated code. */
+	EGeneratedCodeVersion GeneratedCodeVersion;
+
+	/** Returns true if module headers were modified since last code generation. */
+	bool NeedsRegeneration() const;
+
+	/** Returns true if modules are compatible. Used to determine if module data can be loaded from makefile. */
+	bool IsCompatibleWith(const FManifestModule& ManifestModule);
+
+	friend FArchive& operator<<(FArchive& Ar, FManifestModule& ManifestModule)
+	{
+		Ar << ManifestModule.Name;
+		Ar << ManifestModule.ModuleType;
+		Ar << ManifestModule.OverrideModuleType;
+		Ar << ManifestModule.LongPackageName;
+		Ar << ManifestModule.BaseDirectory;
+		Ar << ManifestModule.IncludeBase;
+		Ar << ManifestModule.GeneratedIncludeDirectory;
+		Ar << ManifestModule.PublicUObjectClassesHeaders;
+		Ar << ManifestModule.PublicUObjectHeaders;
+		Ar << ManifestModule.PrivateUObjectHeaders;
+		Ar << ManifestModule.GeneratedCPPFilenameBase;
+		Ar << ManifestModule.SaveExportedHeaders;
+		Ar << ManifestModule.GeneratedCodeVersion;
+
+		return Ar;
+	}
+
+	bool ShouldForceRegeneration() const
+	{
+		return bForceRegeneration;
+	}
+
+	void ForceRegeneration()
+	{
+		bForceRegeneration = true;
+	}
+
+private:
+	bool bForceRegeneration;
+};
+
+
+struct FManifest
+{
+	bool    IsGameTarget;
+	FString RootLocalPath;
+	FString RootBuildPath;
+	FString TargetName;
+	FString ExternalDependenciesFile;
+	
+	/** Ordered list of modules that define UObjects or UStructs, which we may need to generate
+	    code for.  The list is in module dependency order, such that most dependent modules appear first. */
+	TArray<FManifestModule> Modules;
+
+	/**
+	 * Loads a *.uhtmanifest from the specified filename.
+	 *
+	 * @param Filename The filename of the manifest to load.
+	 * @return The loaded module info.
+	 */
+	static FManifest LoadFromFile(const FString& Filename);
+
+	friend FArchive& operator<<(FArchive& Ar, FManifest& Manifest)
+	{
+		Ar << Manifest.IsGameTarget;
+		Ar << Manifest.RootLocalPath;
+		Ar << Manifest.RootBuildPath;
+		Ar << Manifest.TargetName;
+		Ar << Manifest.Modules;
+
+		return Ar;
+	}
+};
diff --git a/Engine/Source/Programs/UHTLite/Private/NativeClassExporter.h b/Engine/Source/Programs/UHTLite/Private/NativeClassExporter.h
new file mode 100644
--- /dev/null
+++ b/Engine/Source/Programs/UHTLite/Private/NativeClassExporter.h
@@ -0,0 +1,570 @@
+// Copyright Epic Games, Inc. All Rights Reserved.
+
+#pragma once
+
+#include "CoreMinimal.h"
+#include "Classes.h"
+#include "Async/TaskGraphInterfaces.h"
+
+class FUnrealSourceFile;
+class UPackage;
+class FProperty;
+class UFunction;
+class UStruct;
+class UField;
+class UClass;
+class UEnum;
+class UScriptStruct;
+class UDelegateFunction;
+class FClassMetaData;
+class FOutputDevice;
+struct FFuncInfo;
+struct FPreloadHeaderFileInfo;
+
+//
+//	FNativeClassHeaderGenerator
+//
+
+namespace EExportFunctionHeaderStyle
+{
+	enum Type
+	{
+		Definition,
+		Declaration
+	};
+}
+
+namespace EExportFunctionType
+{
+	enum Type
+	{
+		Interface,
+		Function,
+		Event
+	};
+}
+
+class FClass;
+class FClasses;
+class FModuleClasses;
+class FScope;
+
+// These are declared in this way to allow swapping out the classes for something more optimized in the future
+typedef FStringOutputDevice FUHTStringBuilder;
+typedef FStringOutputDeviceCountLines FUHTStringBuilderLineCounter;
+
+enum class EExportingState
+{
+	Normal,
+	TypeEraseDelegates
+};
+
+enum class EExportCallbackType
+{
+	Interface,
+	Class
+};
+
+struct FPropertyNamePointerPair
+{
+	FPropertyNamePointerPair(FString InName, FProperty* InProp)
+		: Name(MoveTemp(InName))
+		, Prop(InProp)
+	{
+	}
+
+	FString Name;
+	FProperty* Prop;
+};
+
+FString CreateUTF8LiteralString(const FString& Str);
+
+struct FNativeClassHeaderGenerator
+{
+private:
+	FString API;
+	FString APIStringPrivate;
+
+	/**
+	 * Gets API string for this header.
+	 */
+	const FString& GetAPIString() const
+	{
+		return APIStringPrivate;
+	}
+
+	const UPackage* Package;
+
+	/** A collection of structures used to gather various kinds of references conveniently grouped together to make passing easier */
+	struct FReferenceGatherers
+	{
+		FReferenceGatherers(TSet<FString>* InUniqueCrossModuleReferences, TSet<FString>& InPackageHeaderPaths, TArray<FString>& InTempHeaderPaths)
+			: UniqueCrossModuleReferences(InUniqueCrossModuleReferences)
+			, PackageHeaderPaths(InPackageHeaderPaths)
+			, TempHeaderPaths(InTempHeaderPaths)
+		{
+		}
+
+		/** Set of already exported cross-module references, to prevent duplicates */
+		TSet<FString>* UniqueCrossModuleReferences;
+		/** Array of all header filenames from the current package. */
+		TSet<FString>& PackageHeaderPaths;
+		/** Array of temp filenames that for files to overwrite headers */
+		TArray<FString>& TempHeaderPaths;
+		/** Forward declarations that we need. */
+		TSet<FString> ForwardDeclarations;
+
+	};
+
+	/** the existing disk version of the header for this package's names */
+	FString OriginalNamesHeader;
+
+	/** References to the tasks to save the temp files so as to be able to wait and be sure they are complete before moving them */
+	FGraphEventArray TempSaveTasks;
+
+	/** If false, exported headers will not be saved to disk */
+	bool bAllowSaveExportedHeaders;
+
+	/**
+	 * Exports the struct's C++ properties to the HeaderText output device and adds special
+	 * compiler directives for GCC to pack as we expect.
+	 *
+	 * @param	Out				alternate output device
+	 * @param	Struct			UStruct to export properties
+	 * @param	TextIndent		Current text indentation
+	 */
+	static void ExportProperties(FOutputDevice& Out, UStruct* Struct, int32 TextIndent);
+
+	/** Return the name of the singleton function that returns the UObject for Item */
+	static const FString& GetPackageSingletonName(const UPackage* Item, TSet<FString>* UniqueCrossModuleReferences);
+
+	/** Return the name of the singleton function that returns the UObject for Item */
+	static const FString& GetSingletonName(UField* Item, TSet<FString>* UniqueCrossModuleReferences, bool bRequiresValidObject=true);
+
+	/** Return the address of the singleton function - handles nullptr */
+	static FString GetSingletonNameFuncAddr(UField* Item, TSet<FString>* UniqueCrossModuleReferences, bool bRequiresValidObject=true);
+
+	/**
+	 * Returns the name (overridden if marked up) with TEXT("") or "" wrappers for use in a string literal.
+	 */
+	template <typename T>
+	static FString GetOverriddenNameForLiteral(const T* Item)
+	{
+		const FString& OverriddenName = Item->GetMetaData(TEXT("OverrideNativeName"));
+		if (!OverriddenName.IsEmpty())
+		{
+			return TEXT("TEXT(\"") + OverriddenName + TEXT("\")");
+		}
+		return TEXT("\"") + Item->GetName() + TEXT("\"");
+	}
+
+	/**
+	 * Returns the name (overridden if marked up) or "" wrappers for use in a string literal.
+	 */
+	template <typename T>
+	static FString GetUTF8OverriddenNameForLiteral(const T* Item)
+	{
+		const FString& OverriddenName = Item->GetMetaData(TEXT("OverrideNativeName"));
+		if (!OverriddenName.IsEmpty())
+		{
+			return CreateUTF8LiteralString(OverriddenName);
+		}
+		return CreateUTF8LiteralString(Item->GetName());
+	}
+
+	/**
+	 * Export functions used to find and call C++ or script implementation of a script function in the interface 
+	 */
+	void ExportInterfaceCallFunctions(FOutputDevice& OutCpp, FUHTStringBuilder& Out, FReferenceGatherers& OutReferenceGatherers, const TArray<UFunction*>& CallbackFunctions, const TCHAR* ClassName) const;
+
+	/**
+	 * Export UInterface boilerplate.
+	 *
+	 * @param UInterfaceBoilerplate Device to export to.
+	 * @param Class Interface to export.
+	 * @param FriendText Friend text for this boilerplate.
+	 */
+	static void ExportUInterfaceBoilerplate(FUHTStringBuilder& UInterfaceBoilerplate, FClass* Class, const FString& FriendText);
+
+public:
+
+	enum class EExportClassOutFlags
+	{
+		None = 0x0,
+		NeedsPushModelHeaders = 0x1 << 0,
+	};
+
+private:
+
+	/**
+	 * After all of the dependency checking, and setup for isolating the generated code, actually export the class
+	 *
+	 * @param	OutEnums			Output device for enums declarations.
+	 * @param	OutputGetter		The function to call to get the output.
+	 * @param	SourceFile			Source file to export.
+	 */
+	void ExportClassFromSourceFileInner(
+		FOutputDevice&           OutGeneratedHeaderText,
+		FOutputDevice&           OutputGetter,
+		FOutputDevice&           OutDeclarations,
+		FReferenceGatherers&     OutReferenceGatherers, 
+		FClass*                  Class,
+		const FUnrealSourceFile& SourceFile,
+		EExportClassOutFlags&    OutFlags
+	) const;
+
+	/**
+	 * After all of the dependency checking, but before actually exporting the class, set up the generated code
+	 *
+	 * @param	SourceFile			Source file to export.
+	 */
+	bool WriteHeader(const FPreloadHeaderFileInfo& FileInfo, const FString& InBodyText, const TSet<FString>& InAdditionalHeaders, FReferenceGatherers& InOutReferenceGatherers, FGraphEventRef& OutTempSaveTask) const;
+
+	/**
+	 * Returns a string in the format CLASS_Something|CLASS_Something which represents all class flags that are set for the specified
+	 * class which need to be exported as part of the DECLARE_CLASS macro
+	 */
+	static FString GetClassFlagExportText( UClass* Class );
+
+	/**
+	 * Exports the header text for the enum specified
+	 * 
+	 * @param	Out		the output device for the mirror struct
+	 * @param	Enums	the enum to export
+	 */
+	void ExportEnum(FOutputDevice& Out, UEnum* Enum) const;
+
+	/**
+	 * Exports the inl text for enums declared in non-UClass headers.
+	 * 
+	 * @param	OutputGetter	The function to call to get the output.
+	 * @param	Enum			the enum to export
+	 */
+	void ExportGeneratedEnumInitCode(FOutputDevice& Out, FReferenceGatherers& OutReferenceGatherers, const FUnrealSourceFile& SourceFile, UEnum* Enum) const;
+
+	/**
+	 * Exports the macro declarations for GENERATED_BODY() for each Foo in the struct specified
+	 * 
+	 * @param	Out				output device
+	 * @param	Struct			The struct to export
+	 */
+	void ExportGeneratedStructBodyMacros(FOutputDevice& OutGeneratedHeaderText, FOutputDevice& Out, FReferenceGatherers& OutReferenceGatherers, const FUnrealSourceFile& SourceFile, UScriptStruct* Struct) const;
+
+	/**
+	 * Exports a local mirror of the specified struct; used to get offsets
+	 * 
+	 * @param	Out			the output device for the mirror struct
+	 * @param	Structs		the struct to export
+	 * @param	TextIndent	the current indentation of the header exporter
+	 */
+	static void ExportMirrorsForNoexportStruct(FOutputDevice& Out, UScriptStruct* Struct, int32 TextIndent);
+
+	/**heade
+	 * Exports the parameter struct declarations for the list of functions specified
+	 * 
+	 * @param	Function	the function that (may) have parameters which need to be exported
+	 * @return	true		if the structure generated is not completely empty
+	 */
+	static bool WillExportEventParms( UFunction* Function );
+
+	/**
+	 * Exports C++ type declarations for delegates
+	 *
+	 * @param	Out					output device
+	 * @param	SourceFile			Source file of the delegate.
+	 * @param	DelegateFunctions	the functions that have parameters which need to be exported
+	 */
+	void ExportDelegateDeclaration(FOutputDevice& Out, FReferenceGatherers& OutReferenceGatherers, const FUnrealSourceFile& SourceFile, UFunction* Function) const;
+
+	/**
+	 * Exports C++ type definitions for delegates
+	 *
+	 * @param	Out					output device
+	 * @param	SourceFile			Source file of the delegate.
+	 * @param	DelegateFunctions	the functions that have parameters which need to be exported
+	 */
+	void ExportDelegateDefinition(FOutputDevice& Out, FReferenceGatherers& OutReferenceGatherers, const FUnrealSourceFile& SourceFile, UFunction* Function) const;
+
+	/**
+	 * Exports the parameter struct declarations for the given function.
+	 *
+	 * @param	Out					output device
+	 * @param	Function			the function that have parameters which need to be exported
+	 * @param	Indent				number of spaces to put before each line
+	 * @param	bOutputConstructor	If true, output a constructor for the param struct
+	 */
+	static void ExportEventParm(FUHTStringBuilder& Out, TSet<FString>& PropertyFwd, UFunction* Function, int32 Indent, bool bOutputConstructor, EExportingState ExportingState);
+
+	/**
+	* Move the temp header files into the .h files
+	* 
+	* @param	PackageName	Name of the package being saved
+	* @param	TempHeaderPaths	Names of all the headers to move
+	*/
+	static void ExportUpdatedHeaders(FString&& PackageName, TArray<FString>&& TempHeaderPaths, FGraphEventArray& InTempSaveTasks);
+
+	/**
+	 * Exports the generated cpp file for all functions/events/delegates in package.
+	 */
+	static void ExportGeneratedCPP(FOutputDevice& Out, const TSet<FString>& InCrossModuleReferences, const TCHAR* EmptyLinkFunctionPostfix, const TCHAR* Body, const TCHAR* OtherIncludes);
+
+	/**
+	 * Get the intrinsic null value for this property
+	 * 
+	 * @param	Prop				the property to get the null value for
+	 * @param	bMacroContext		true when exporting the P_GET* macro, false when exporting the friendly C++ function header
+	 * @param	bInitializer		if true, will just return ForceInit instead of FStruct(ForceInit)
+	 *
+	 * @return	the intrinsic null value for the property (0 for ints, TEXT("") for strings, etc.)
+	 */
+	static FString GetNullParameterValue( FProperty* Prop, bool bInitializer = false );
+
+	/**
+	 * Exports a native function prototype
+	 * 
+	 * @param	Out					Where to write the exported function prototype to.
+	 * @param	FunctionData		data representing the function to export
+	 * @param	FunctionType		Whether to export this function prototype as an event stub, an interface or a native function stub.
+	 * @param	FunctionHeaderStyle	Whether we're outputting a declaration or definition.
+	 * @param	ExtraParam			Optional extra parameter that will be added to the declaration as the first argument
+	 */
+	static void ExportNativeFunctionHeader(
+		FOutputDevice&                   Out,
+		TSet<FString>&                   OutFwdDecls,
+		const FFuncInfo&                 FunctionData,
+		EExportFunctionType::Type        FunctionType,
+		EExportFunctionHeaderStyle::Type FunctionHeaderStyle,
+		const TCHAR*                     ExtraParam,
+		const TCHAR*                     APIString
+	);
+
+	/**
+	* Runs checks whether necessary RPC functions exist for function described by FunctionData.
+	*
+	* @param	FunctionData			Data representing the function to export.
+	* @param	ClassName				Name of currently parsed class.
+	* @param	ImplementationPosition	Position in source file of _Implementation function for function described by FunctionData.
+	* @param	ValidatePosition		Position in source file of _Validate function for function described by FunctionData.
+	* @param	SourceFile				Currently analyzed source file.
+	*/
+	void CheckRPCFunctions(FReferenceGatherers& OutReferenceGatherers, const FFuncInfo& FunctionData, const FString& ClassName, int32 ImplementationPosition, int32 ValidatePosition, const FUnrealSourceFile& SourceFile) const;
+
+	/**
+	 * Exports the native stubs for the list of functions specified
+	 * 
+	 * @param SourceFile	current source file
+	 * @param Class			class
+	 * @param ClassData		class data
+	 */
+	void ExportNativeFunctions(FOutputDevice& OutGeneratedHeaderText, FOutputDevice& OutGeneratedCPPText, FOutputDevice& OutMacroCalls, FOutputDevice& OutNoPureDeclsMacroCalls, FReferenceGatherers& OutReferenceGatherers, const FUnrealSourceFile& SourceFile, UClass* Class, FClassMetaData* ClassData) const;
+
+	/**
+	 * Export the actual internals to a standard thunk function
+	 *
+	 * @param RPCWrappers output device for writing
+	 * @param Function given function
+	 * @param FunctionData function data for the current function
+	 * @param Parameters list of parameters in the function
+	 * @param Return return parameter for the function
+	 */
+	void ExportFunctionThunk(FUHTStringBuilder& RPCWrappers, FReferenceGatherers& OutReferenceGatherers, UFunction* Function, const FFuncInfo& FunctionData, const TArray<FProperty*>& Parameters, FProperty* Return) const;
+
+	/** Exports the native function registration code for the given class. */
+	static void ExportNatives(FOutputDevice& Out, FClass* Class);
+
+	/**
+	 * Exports generated singleton functions for UObjects that used to be stored in .u files.
+	 * 
+	 * @param	Out				The destination to write to.
+	 * @param	SourceFile		The source file being processed.
+	 * @param	Class			Class to export
+	 * @param	OutFriendText	(Output parameter) Friend text
+	 */
+	void ExportNativeGeneratedInitCode(FOutputDevice& Out, FOutputDevice& OutDeclarations, FReferenceGatherers& OutReferenceGatherers, const FUnrealSourceFile& SourceFile, FClass* Class, FUHTStringBuilder& OutFriendText) const;
+
+	/**
+	 * Export given function.
+	 *
+	 * @param	Out				The destination to write to.
+	 * @param	Function		Given function.
+	 * @param	bIsNoExport		Is in NoExport class.
+	 */
+	void ExportFunction(FOutputDevice& Out, FReferenceGatherers& OutReferenceGatherers, const FUnrealSourceFile& SourceFile, UFunction* Function, bool bIsNoExport) const;
+
+	/**
+	 * Exports a generated singleton function to setup the package for compiled-in classes.
+	 * 
+	 * @param	Out			The destination to write to.
+	 * @param	Package		Package to export code for.
+	**/
+	void ExportGeneratedPackageInitCode(FOutputDevice& Out, const TCHAR* InDeclarations, const UPackage* Package, uint32 CRC);
+
+	/**
+	 * Function to output the C++ code necessary to set up the given array of properties
+	 * 
+	 * @param	DeclOut			String output device to send the generated declarations to
+	 * @param	Out				String output device to send the generated code to
+	 * @param	Scope			The scope to prefix on all variable definitions
+	 * @param	Properties		Array of properties to export
+	 * @param	Spaces			String of spaces to use as an indent for the declaration
+	 * @param	Spaces			String of spaces to use as an indent
+	 *
+	 * @return      A pair of strings which represents the pointer and a count of the emitted properties.
+	 */
+	TTuple<FString, FString> OutputProperties(FOutputDevice& DeclOut, FOutputDevice& Out, FReferenceGatherers& OutReferenceGatherers, const TCHAR* Scope, const TArray<FProperty*>& Properties, const TCHAR* DeclSpaces, const TCHAR* Spaces) const;
+
+	/**
+	 * Function to output the C++ code necessary to set up a property
+	 * 
+	 * @param	DeclOut			String output device to send the generated declarations to
+	 * @param	Out				String output device to send the generated code to
+	 * @param	Scope			The scope to prefix on all variable definitions
+	 * @param	Prop			Property to export
+	 * @param	DeclSpaces		String of spaces to use as an indent for the declaration
+	 * @param	Spaces			String of spaces to use as an indent
+	**/
+	void OutputProperty(FOutputDevice& DeclOut, FOutputDevice& Out, FReferenceGatherers& OutReferenceGatherers, const TCHAR* Scope, TArray<FPropertyNamePointerPair>& PropertyNamesAndPointers, FProperty* Prop, const TCHAR* DeclSpaces, const TCHAR* Spaces) const;
+
+	/**
+	 * Function to output the C++ code necessary to set up a property, including an array property and its inner, array dimensions, etc.
+	 *
+	 * @param	DeclOut			The destination to write declarations to.
+	 * @param	Out				The destination to write to.
+	 * @param	Prop			Property to export
+	 * @param	OffsetStr		String specifying the property offset
+	 * @param	Name			Name for the generated variable
+	 * @param	DeclSpaces		String of spaces to use as an indent for the declaration
+	 * @param	Spaces			String of spaces to use as an indent
+	 * @param	SourceStruct	Structure that the property offset is relative to
+	**/
+	void PropertyNew(FOutputDevice& DeclOut, FOutputDevice& Out, FReferenceGatherers& OutReferenceGatherers, FProperty* Prop, const TCHAR* OffsetStr, const TCHAR* Name, const TCHAR* DeclSpaces, const TCHAR* Spaces, const TCHAR* SourceStruct = nullptr) const;
+
+	/**
+	 * Exports the proxy definitions for the list of enums specified
+	 * 
+	 * @param SourceFile	current source file
+	 */
+	static void ExportCallbackFunctions(
+		FOutputDevice&            OutGeneratedHeaderText,
+		FOutputDevice&            Out,
+		TSet<FString>&            OutFwdDecls,
+		const TArray<UFunction*>& CallbackFunctions,
+		const TCHAR*              CallbackWrappersMacroName,
+		EExportCallbackType       ExportCallbackType,
+		const TCHAR*              APIString
+	);
+
+	/**
+	 * Determines if the property has alternate export text associated with it and if so replaces the text in PropertyText with the
+	 * alternate version. (for example, structs or properties that specify a native type using export-text).  Should be called immediately
+	 * after ExportCppDeclaration()
+	 *
+	 * @param	Prop			the property that is being exported
+	 * @param	PropertyText	the string containing the text exported from ExportCppDeclaration
+	 */
+	static void ApplyAlternatePropertyExportText(FProperty* Prop, FUHTStringBuilder& PropertyText, EExportingState ExportingState);
+
+	/**
+	* Create a temp header file name from the header name
+	*
+	* @param	CurrentFilename		The filename off of which the current filename will be generated
+	* @param	bReverseOperation	Get the header from the temp file name instead
+	*
+	* @return	The generated string
+	*/
+	static FString GenerateTempHeaderName( const FString& CurrentFilename, bool bReverseOperation = false );
+
+	/**
+	 * Saves a generated header if it has changed. 
+	 *
+	 * @param HeaderPath	Header Filename
+	 * @param NewHeaderContents	Contents of the generated header.
+	 * @return True if the header contents has changed, false otherwise.
+	 */
+	bool SaveHeaderIfChanged(FReferenceGatherers& OutReferenceGatherers, const FPreloadHeaderFileInfo& FileInfo, FString&& NewHeaderContents, FGraphEventRef& OutSaveTaskRef) const;
+
+	/**
+	 * Deletes all .generated.h files which do not correspond to any of the classes.
+	 */
+	static void DeleteUnusedGeneratedHeaders(TSet<FString>&& PackageHeaderPathSet);
+
+	/**
+	 * Exports macros that manages UObject constructors.
+	 * 
+	 * @param	VTableOut								The destination to write vtable helpers to.
+	 * @param	StandardUObjectConstructorsMacroCall	The destination to write standard constructor macros to.
+	 * @param	EnhancedUObjectConstructorsMacroCall	The destination to write enhanced constructor macros to.
+	 * @param	ConstructorsMacroPrefix					Prefix for constructors macro.
+	 * @param	Class									Class for which to export macros.
+	 */
+	static void ExportConstructorsMacros(FOutputDevice& OutGeneratedHeaderText, FOutputDevice& VTableOut, FOutputDevice& StandardUObjectConstructorsMacroCall, FOutputDevice& EnhancedUObjectConstructorsMacroCall, const FString& ConstructorsMacroPrefix, FClass* Class, const TCHAR* APIArg);
+
+public:
+
+	/** 
+	 * Properties in source files generated from blueprint assets have a symbol name that differs from the source asset.
+	 * This function returns the original name of the field (rather than the native, symbol name).
+	 */
+	template <typename T>
+	static FString GetOverriddenName(const T* Item)
+	{
+		const FString& OverriddenName = Item->GetMetaData(TEXT("OverrideNativeName"));
+		if (!OverriddenName.IsEmpty())
+		{
+			return OverriddenName.ReplaceCharWithEscapedChar();
+		}
+		return Item->GetName();
+	}
+
+	template <typename T>
+	static FName GetOverriddenFName(const T* Item)
+	{
+		FString OverriddenName = Item->GetMetaData(TEXT("OverrideNativeName"));
+		if (!OverriddenName.IsEmpty())
+		{
+			return FName(*OverriddenName);
+		}
+		return Item->GetFName();
+	}
+
+	template <typename T>
+	static FString GetOverriddenPathName(const T* Item)
+	{
+		return FString::Printf(TEXT("%s.%s"), *FClass::GetTypePackageName(Item), *GetOverriddenName(Item));
+	}
+
+	// Constructor
+	FNativeClassHeaderGenerator(
+		const UPackage* InPackage,
+		const TSet<FUnrealSourceFile*>& SourceFiles,
+		FClasses& AllClasses,
+		bool InAllowSaveExportedHeaders
+	);
+
+	/**
+	 * Gets string with function return type.
+	 * 
+	 * @param Function Function to get return type of.
+	 * @return FString with function return type.
+	 */
+	static FString GetFunctionReturnString(UFunction* Function, FReferenceGatherers& OutReferenceGatherers);
+
+	/**
+	* Gets string with function parameters (with names).
+	*
+	* @param Function Function to get parameters of.
+	* @return FString with function parameters.
+	*/
+	static FString GetFunctionParameterString(UFunction* Function, FReferenceGatherers& OutReferenceGatherers);
+
+	/**
+	 * Checks if function is missing "virtual" specifier.
+	 * 
+	 * @param SourceFile SourceFile where function is declared.
+	 * @param FunctionNamePosition Position of name of function in SourceFile.
+	 * @return true if function misses "virtual" specifier, false otherwise.
+	 */
+	static bool IsMissingVirtualSpecifier(const FString& SourceFile, int32 FunctionNamePosition);
+};
+
+ENUM_CLASS_FLAGS(FNativeClassHeaderGenerator::EExportClassOutFlags);
\ No newline at end of file
diff --git a/Engine/Source/Programs/UHTLite/Private/ParserClass.h b/Engine/Source/Programs/UHTLite/Private/ParserClass.h
new file mode 100644
--- /dev/null
+++ b/Engine/Source/Programs/UHTLite/Private/ParserClass.h
@@ -0,0 +1,71 @@
+// Copyright Epic Games, Inc. All Rights Reserved.
+
+#pragma once
+
+#include "CoreMinimal.h"
+#include "UObject/Class.h"
+#include "UObject/Field.h"
+
+class UField;
+
+struct EEnforceInterfacePrefix
+{
+	enum Type
+	{
+		None,
+		I,
+		U
+	};
+};
+
+class FClasses;
+
+class FClass : public UClass
+{
+public:
+	FClass();
+
+	/** 
+	 * Returns the name of the given class with a valid prefix.
+	 *
+	 * @param InClass Class used to create a valid class name with prefix
+	 */
+	FString GetNameWithPrefix(EEnforceInterfacePrefix::Type EnforceInterfacePrefix = EEnforceInterfacePrefix::None) const;
+
+	/**
+	 * Returns the super class of this class, or NULL if there is no superclass.
+	 *
+	 * @return The super class of this class.
+	 */
+	FClass* GetSuperClass() const;
+
+	/**
+	 * Returns the 'within' class of this class.
+	 *
+	 * @return The 'within' class of this class.
+	 */
+	FClass* GetClassWithin() const;
+
+	TArray<FClass*> GetInterfaceTypes() const;
+
+	void GetHideCategories(TArray<FString>& OutHideCategories) const;
+	void GetShowCategories(TArray<FString>& OutShowCategories) const;
+	void GetSparseClassDataTypes(TArray<FString>& OutSparseClassDataTypes) const;
+
+	/** Helper function that checks if the field is a dynamic type (can be constructed post-startup) */
+	template <typename T>
+	static bool IsDynamic(const T* Field)
+	{
+		return Field->HasMetaData(NAME_ReplaceConverted);
+	}
+
+	/** Helper function that checks if the field is belongs to a dynamic type */
+	static bool IsOwnedByDynamicType(const UField* Field);
+	static bool IsOwnedByDynamicType(const FField* Field);
+
+	/** Helper function to get the source replaced package name */
+	static const FString& GetTypePackageName(const UField* Field);
+	static const FString& GetTypePackageName(const FField* Field);
+
+	static const FName NAME_ReplaceConverted;
+};
diff --git a/Engine/Source/Programs/UHTLite/Private/ParserHelper.cpp b/Engine/Source/Programs/UHTLite/Private/ParserHelper.cpp
new file mode 100644
--- /dev/null
+++ b/Engine/Source/Programs/UHTLite/Private/ParserHelper.cpp
@@ -0,0 +1,325 @@
+// Copyright Epic Games, Inc. All Rights Reserved.
+
+
+#include "ParserHelper.h"
+#include "UnrealHeaderTool.h"
+#include "Algo/Find.h"
+#include "Misc/DefaultValueHelper.h"
+
+/////////////////////////////////////////////////////
+// FClassMetaData
+
+/**
+ * Finds the metadata for the property specified
+ *
+ * @param	Prop	the property to search for
+ *
+ * @return	pointer to the metadata for the property specified, or NULL
+ *			if the property doesn't exist in the list (for example, if it
+ *			is declared in a package that is already compiled and has had its
+ *			source stripped)
+ */
+FTokenData* FClassMetaData::FindTokenData( FProperty* Prop )
+{
+	check(Prop);
+
+	FTokenData* Result = nullptr;
+	UObject* Outer = Prop->GetOwner<UObject>();
+	check(Outer);
+	UClass* OuterClass = nullptr;
+	if (Outer->IsA<UStruct>())
+	{
+		Result = GlobalPropertyData.Find(Prop);
+
+		if (Result == nullptr)
+		{
+			OuterClass = Cast<UClass>(Outer);
+
+			if (OuterClass != nullptr && OuterClass->GetSuperClass() != OuterClass)
+			{
+				OuterClass = OuterClass->GetSuperClass();
+			}
+		}
+	}
+	else
+	{
+		UFunction* OuterFunction = Cast<UFunction>(Outer);
+		if ( OuterFunction != NULL )
+		{
+			// function parameter, return, or local property
+			FFunctionData* FuncData = nullptr;
+			if (FFunctionData::TryFindForFunction(OuterFunction, FuncData))
+			{
+				FPropertyData& FunctionParameters = FuncData->GetParameterData();
+				Result = FunctionParameters.Find(Prop);
+				if ( Result == NULL )
+				{
+					Result = FuncData->GetReturnTokenData();
+				}
+			}
+			else
+			{
+				OuterClass = OuterFunction->GetOwnerClass();
+			}
+		}
+		else
+		{
+			// struct property
+			UScriptStruct* OuterStruct = Cast<UScriptStruct>(Outer);
+			check(OuterStruct != NULL);
+			OuterClass = OuterStruct->GetOwnerClass();
+		}
+	}
+
+	if (Result == nullptr && OuterClass != nullptr)
+	{
+		FClassMetaData* SuperClassData = GScriptHelper.FindClassData(OuterClass);
+		if (SuperClassData && SuperClassData != this)
+		{
+			Result = SuperClassData->FindTokenData(Prop);
+		}
+	}
+
+	return Result;
+}
+
+void FClassMetaData::AddInheritanceParent(FString&& InParent, FUnrealSourceFile* UnrealSourceFile)
+{
+	MultipleInheritanceParents.Add(new FMultipleInheritanceBaseClass(MoveTemp(InParent)));
+}
+
+void FClassMetaData::AddInheritanceParent(UClass* ImplementedInterfaceClass, FUnrealSourceFile* UnrealSourceFile)
+{
+	MultipleInheritanceParents.Add(new FMultipleInheritanceBaseClass(ImplementedInterfaceClass));
+}
+
+/////////////////////////////////////////////////////
+// FPropertyBase
+
+const TCHAR* FPropertyBase::GetPropertyTypeText( EPropertyType Type )
+{
+	switch ( Type )
+	{
+		CASE_TEXT(CPT_None);
+		CASE_TEXT(CPT_Byte);
+		CASE_TEXT(CPT_Int8);
+		CASE_TEXT(CPT_Int16);
+		CASE_TEXT(CPT_Int);
+		CASE_TEXT(CPT_Int64);
+		CASE_TEXT(CPT_UInt16);
+		CASE_TEXT(CPT_UInt32);
+		CASE_TEXT(CPT_UInt64);
+		CASE_TEXT(CPT_Bool);
+		CASE_TEXT(CPT_Bool8);
+		CASE_TEXT(CPT_Bool16);
+		CASE_TEXT(CPT_Bool32);
+		CASE_TEXT(CPT_Bool64);
+		CASE_TEXT(CPT_Float);
+		CASE_TEXT(CPT_Double);
+		CASE_TEXT(CPT_ObjectReference);
+		CASE_TEXT(CPT_Interface);
+		CASE_TEXT(CPT_Name);
+		CASE_TEXT(CPT_Delegate);
+		CASE_TEXT(CPT_Struct);
+		CASE_TEXT(CPT_String);
+		CASE_TEXT(CPT_Text);
+		CASE_TEXT(CPT_MulticastDelegate);
+		CASE_TEXT(CPT_SoftObjectReference);
+		CASE_TEXT(CPT_WeakObjectReference);
+		CASE_TEXT(CPT_LazyObjectReference);
+		CASE_TEXT(CPT_Map);
+		CASE_TEXT(CPT_Set);
+		CASE_TEXT(CPT_FieldPath);
+		CASE_TEXT(CPT_MAX);
+	}
+
+	return TEXT("");
+}
+
+/////////////////////////////////////////////////////
+// FToken
+
+/**
+ * Copies the properties from this token into another.
+ *
+ * @param	Other	the token to copy this token's properties to.
+ */
+FToken& FToken::operator=(const FToken& Other)
+{
+	FPropertyBase::operator=((FPropertyBase&)Other);
+
+	TokenType = Other.TokenType;
+	TokenName = Other.TokenName;
+	bTokenNameInitialized = Other.bTokenNameInitialized;
+	StartPos = Other.StartPos;
+	StartLine = Other.StartLine;
+	TokenProperty = Other.TokenProperty;
+
+	FCString::Strncpy(Identifier, Other.Identifier, NAME_SIZE);
+	FMemory::Memcpy(String, Other.String, sizeof(String));
+
+	return *this;
+}
+
+FToken& FToken::operator=(FToken&& Other)
+{
+	FPropertyBase::operator=(MoveTemp(Other));
+
+	TokenType = Other.TokenType;
+	TokenName = Other.TokenName;
+	bTokenNameInitialized = Other.bTokenNameInitialized;
+	StartPos = Other.StartPos;
+	StartLine = Other.StartLine;
+	TokenProperty = Other.TokenProperty;
+
+	FCString::Strncpy(Identifier, Other.Identifier, NAME_SIZE);
+	FMemory::Memcpy(String, Other.String, sizeof(String));
+
+	return *this;
+}
+
+/////////////////////////////////////////////////////
+// FAdvancedDisplayParameterHandler
+static const FName NAME_AdvancedDisplay(TEXT("AdvancedDisplay"));
+
+FAdvancedDisplayParameterHandler::FAdvancedDisplayParameterHandler(const TMap<FName, FString>* MetaData)
+	: NumberLeaveUnmarked(-1), AlreadyLeft(0), bUseNumber(false)
+{
+	if(MetaData)
+	{
+		const FString* FoundString = MetaData->Find(NAME_AdvancedDisplay);
+		if(FoundString)
+		{
+			FoundString->ParseIntoArray(ParametersNames, TEXT(","), true);
+			for(int32 NameIndex = 0; NameIndex < ParametersNames.Num();)
+			{
+				FString& ParameterName = ParametersNames[NameIndex];
+				ParameterName.TrimStartAndEndInline();
+				if(ParameterName.IsEmpty())
+				{
+					ParametersNames.RemoveAtSwap(NameIndex);
+				}
+				else
+				{
+					++NameIndex;
+				}
+			}
+			if(1 == ParametersNames.Num())
+			{
+				bUseNumber = FDefaultValueHelper::ParseInt(ParametersNames[0], NumberLeaveUnmarked);
+			}
+		}
+	}
+}
+
+bool FAdvancedDisplayParameterHandler::ShouldMarkParameter(const FString& ParameterName)
+{
+	if(bUseNumber)
+	{
+		if(NumberLeaveUnmarked < 0)
+		{
+			return false;
+		}
+		if(AlreadyLeft < NumberLeaveUnmarked)
+		{
+			AlreadyLeft++;
+			return false;
+		}
+		return true;
+	}
+	return ParametersNames.Contains(ParameterName);
+}
+
+bool FAdvancedDisplayParameterHandler::CanMarkMore() const
+{
+	return bUseNumber ? (NumberLeaveUnmarked > 0) : (0 != ParametersNames.Num());
+}
+
+TMap<UFunction*, TUniqueObj<FFunctionData> > FFunctionData::FunctionDataMap;
+
+FFunctionData* FFunctionData::FindForFunction(UFunction* Function)
+{
+	TUniqueObj<FFunctionData>* Output = FunctionDataMap.Find(Function);
+
+	check(Output);
+
+	return &(*Output).Get();
+}
+
+FFunctionData* FFunctionData::Add(UFunction* Function)
+{
+	TUniqueObj<FFunctionData>& Output = FunctionDataMap.Add(Function);
+
+	return &Output.Get();
+}
+
+FFunctionData* FFunctionData::Add(FFuncInfo&& FunctionInfo)
+{
+	TUniqueObj<FFunctionData>& Output = FunctionDataMap.Emplace(FunctionInfo.FunctionReference, MoveTemp(FunctionInfo));
+
+	return &Output.Get();
+}
+
+bool FFunctionData::TryFindForFunction(UFunction* Function, FFunctionData*& OutData)
+{
+	TUniqueObj<FFunctionData>* Output = FunctionDataMap.Find(Function);
+
+	if (!Output)
+	{
+		return false;
+	}
+
+	OutData = &(*Output).Get();
+	return true;
+}
+
+FClassMetaData* FCompilerMetadataManager::AddClassData(UStruct* Struct, FUnrealSourceFile* UnrealSourceFile)
+{
+	TUniquePtr<FClassMetaData>* pClassData = Find(Struct);
+	if (pClassData == NULL)
+	{
+		pClassData = &Emplace(Struct, new FClassMetaData());
+	}
+
+	return pClassData->Get();
+}
+
+FClassMetaData* FCompilerMetadataManager::AddInterfaceClassData(UStruct* Struct, FUnrealSourceFile* UnrealSourceFile)
+{
+	FClassMetaData* ClassData = AddClassData(Struct, UnrealSourceFile);
+	ClassData->ParsedInterface = EParsedInterface::ParsedUInterface;
+	InterfacesToVerify.Emplace(Struct, ClassData);
+	return ClassData;
+}
+
+FTokenData* FPropertyData::Set(FProperty* InKey, FTokenData&& InValue, FUnrealSourceFile* UnrealSourceFile)
+{
+	FTokenData* Result = NULL;
+
+	TSharedPtr<FTokenData>* pResult = Super::Find(InKey);
+	if (pResult != NULL)
+	{
+		Result = pResult->Get();
+		*Result = MoveTemp(InValue);
+	}
+	else
+	{
+		pResult = &Super::Emplace(InKey, new FTokenData(MoveTemp(InValue)));
+		Result = pResult->Get();
+	}
+
+	return Result;
+}
+
+void FCompilerMetadataManager::CheckForNoIInterfaces()
+{
+	for (const TPair<UStruct*, FClassMetaData*>& StructDataPair : InterfacesToVerify)
+	{
+		if (StructDataPair.Value->ParsedInterface == EParsedInterface::ParsedUInterface)
+		{
+			FString Name = StructDataPair.Key->GetName();
+			FError::Throwf(TEXT("UInterface 'U%s' parsed without a corresponding 'I%s'"), *Name, *Name);
+		}
+	}
+	InterfacesToVerify.Reset();
+}
\ No newline at end of file
