Branching UHT to UHTLite:

diff --git a/Engine/Source/Programs/UHTLite/Private/ParserHelper.h b/Engine/Source/Programs/UHTLite/Private/ParserHelper.h
new file mode 100644
--- /dev/null
+++ b/Engine/Source/Programs/UHTLite/Private/ParserHelper.h
@@ -0,0 +1,1923 @@
+// Copyright Epic Games, Inc. All Rights Reserved.
+
+#pragma once
+
+#include "CoreMinimal.h"
+#include "UObject/ErrorException.h"
+#include "UObject/Package.h"
+#include "UObject/MetaData.h"
+#include "UObject/TextProperty.h"
+#include "UObject/EnumProperty.h"
+#include "UObject/FieldPathProperty.h"
+#include "UnrealHeaderToolGlobals.h"
+#include "ClassMaps.h"
+#include "Templates/UniqueObj.h"
+#include "Templates/UniquePtr.h"
+
+class UEnum;
+class UScriptStruct;
+class FProperty;
+class FUnrealSourceFile;
+class UObject;
+class UField;
+class UMetaData;
+class FHeaderParser;
+
+extern class FCompilerMetadataManager GScriptHelper;
+
+/*-----------------------------------------------------------------------------
+	FPropertyBase.
+-----------------------------------------------------------------------------*/
+
+enum EFunctionExportFlags
+{
+	FUNCEXPORT_Final			=0x00000001,	// function declaration included "final" keyword.  Used to differentiate between functions that have FUNC_Final only because they're private
+	//							=0x00000002,
+	//							=0x00000004,
+	FUNCEXPORT_RequiredAPI		=0x00000008,	// Function should be exported as a public API function
+	FUNCEXPORT_Inline			=0x00000010,	// export as an inline static C++ function
+	FUNCEXPORT_CppStatic		=0x00000020,	// Export as a real C++ static function, causing thunks to call via ClassName::FuncName instead of this->FuncName
+	FUNCEXPORT_CustomThunk		=0x00000040,	// Export no thunk function; the user will manually define a custom one
+	//							=0x00000080,
+	//							=0x00000100,
+};
+
+enum EPropertyHeaderExportFlags
+{
+	PROPEXPORT_Public		=0x00000001,	// property should be exported as public
+	PROPEXPORT_Private		=0x00000002,	// property should be exported as private
+	PROPEXPORT_Protected	=0x00000004,	// property should be exported as protected
+};
+
+struct EPointerType
+{
+	enum Type
+	{
+		None,
+		Native
+	};
+};
+
+struct EArrayType
+{
+	enum Type
+	{
+		None,
+		Static,
+		Dynamic,
+		Set
+	};
+};
+
+enum class EAllocatorType
+{
+	Default,
+	MemoryImage
+};
+
+struct ERefQualifier
+{
+	enum Type
+	{
+		None,
+		ConstRef,
+		NonConstRef
+	};
+};
+
+enum class EIntType
+{
+	None,
+	Sized,  // e.g. int32, int16
+	Unsized // e.g. int, unsigned int
+};
+
+#ifndef CASE_TEXT
+#define CASE_TEXT(txt) case txt: return TEXT(#txt)
+#endif
+
+/**
+ * Basic information describing a type.
+ */
+class FPropertyBase : public TSharedFromThis<FPropertyBase>
+{
+public:
+	// Variables.
+	EPropertyType       Type;
+	EArrayType::Type    ArrayType;
+	EAllocatorType      AllocatorType = EAllocatorType::Default;
+	EPropertyFlags      PropertyFlags;
+	EPropertyFlags      ImpliedPropertyFlags;
+	ERefQualifier::Type RefQualifier; // This is needed because of legacy stuff - FString mangles the flags for reasons that have become lost in time but we need this info for testing for invalid replicated function signatures.
+
+	TSharedPtr<FPropertyBase> MapKeyProp;
+
+	/**
+	 * A mask of EPropertyHeaderExportFlags which are used for modifying how this property is exported to the native class header
+	 */
+	uint32 PropertyExportFlags;
+	union
+	{
+		class UEnum* Enum;
+		class UClass* PropertyClass;
+		class UScriptStruct* Struct;
+		class UFunction* Function;
+		class FFieldClass* PropertyPathClass;
+#if PLATFORM_64BITS
+		int64 StringSize;
+#else
+		int32 StringSize;
+#endif
+	};
+
+	UClass* MetaClass;
+	FName   DelegateName;
+	UClass*	DelegateSignatureOwnerClass;
+	FName   RepNotifyName;
+
+	/** Raw string (not type-checked) used for specifying special text when exporting a property to the *Classes.h file */
+	FString	ExportInfo;
+
+	/** Map of key value pairs that will be added to the package's UMetaData for this property */
+	TMap<FName, FString> MetaData;
+
+	EPointerType::Type PointerType;
+	EIntType IntType;
+
+public:
+	/** @name Constructors */
+	//@{
+	FPropertyBase(const FPropertyBase&) = default;
+	FPropertyBase(FPropertyBase&&) = default;
+	FPropertyBase& operator=(const FPropertyBase&) = default;
+	FPropertyBase& operator=(FPropertyBase&&) = default;
+
+	explicit FPropertyBase(EPropertyType InType)
+	: Type                (InType)
+	, ArrayType           (EArrayType::None)
+	, PropertyFlags       (CPF_None)
+	, ImpliedPropertyFlags(CPF_None)
+	, RefQualifier        (ERefQualifier::None)
+	, PropertyExportFlags (PROPEXPORT_Public)
+	, StringSize          (0)
+	, MetaClass           (nullptr)
+	, DelegateName        (NAME_None)
+	, DelegateSignatureOwnerClass(nullptr)
+	, RepNotifyName       (NAME_None)
+	, PointerType         (EPointerType::None)
+	, IntType             (GetSizedIntTypeFromPropertyType(InType))
+	{
+	}
+
+	explicit FPropertyBase(EPropertyType InType, EIntType InIntType)
+	: Type                (InType)
+	, ArrayType           (EArrayType::None)
+	, PropertyFlags       (CPF_None)
+	, ImpliedPropertyFlags(CPF_None)
+	, RefQualifier        (ERefQualifier::None)
+	, PropertyExportFlags (PROPEXPORT_Public)
+	, StringSize          (0)
+	, MetaClass           (nullptr)
+	, DelegateName        (NAME_None)
+	, DelegateSignatureOwnerClass(nullptr)
+	, RepNotifyName       (NAME_None)
+	, PointerType         (EPointerType::None)
+	, IntType             (InIntType)
+	{
+	}
+
+	explicit FPropertyBase(UEnum* InEnum, EPropertyType InType)
+	: Type                (InType)
+	, ArrayType           (EArrayType::None)
+	, PropertyFlags       (CPF_None)
+	, ImpliedPropertyFlags(CPF_None)
+	, RefQualifier        (ERefQualifier::None)
+	, PropertyExportFlags (PROPEXPORT_Public)
+	, Enum                (InEnum)
+	, MetaClass           (nullptr)
+	, DelegateName        (NAME_None)
+	, DelegateSignatureOwnerClass(nullptr)
+	, RepNotifyName       (NAME_None)
+	, PointerType         (EPointerType::None)
+	, IntType             (GetSizedIntTypeFromPropertyType(InType))
+	{
+	}
+
+	explicit FPropertyBase(UClass* InClass, bool bIsWeak = false, bool bWeakIsAuto = false, bool bIsLazy = false, bool bIsSoft = false)
+	: Type                (CPT_ObjectReference)
+	, ArrayType           (EArrayType::None)
+	, PropertyFlags       (CPF_None)
+	, ImpliedPropertyFlags(CPF_None)
+	, RefQualifier        (ERefQualifier::None)
+	, PropertyExportFlags (PROPEXPORT_Public)
+	, PropertyClass       (InClass)
+	, MetaClass           (nullptr)
+	, DelegateName        (NAME_None)
+	, DelegateSignatureOwnerClass(nullptr)
+	, RepNotifyName       (NAME_None)
+	, PointerType         (EPointerType::None)
+	, IntType             (EIntType::None)
+	{
+		// if this is an interface class, we use the FInterfaceProperty class instead of FObjectProperty
+		if ( InClass->HasAnyClassFlags(CLASS_Interface) )
+		{
+			Type = CPT_Interface;
+		}
+		if (bIsLazy)
+		{
+			Type = CPT_LazyObjectReference;
+		}
+		else if (bIsSoft)
+		{
+			Type = CPT_SoftObjectReference;
+		}
+		else if (bIsWeak)
+		{
+			Type = CPT_WeakObjectReference;
+			if (bWeakIsAuto)
+			{
+				PropertyFlags |= CPF_AutoWeak;
+			}
+		}
+	}
+
+	explicit FPropertyBase(UScriptStruct* InStruct)
+	: Type                (CPT_Struct)
+	, ArrayType           (EArrayType::None)
+	, PropertyFlags       (CPF_None)
+	, ImpliedPropertyFlags(CPF_None)
+	, RefQualifier        (ERefQualifier::None)
+	, PropertyExportFlags (PROPEXPORT_Public)
+	, Struct              (InStruct)
+	, MetaClass           (NULL)
+	, DelegateName        (NAME_None)
+	, DelegateSignatureOwnerClass(nullptr)
+	, RepNotifyName       (NAME_None)
+	, PointerType         (EPointerType::None)
+	, IntType             (EIntType::None)
+	{
+	}
+
+	explicit FPropertyBase(FProperty* Property)
+	: PropertyExportFlags(PROPEXPORT_Public)
+	, DelegateName       (NAME_None)
+	, DelegateSignatureOwnerClass(nullptr)
+	, RepNotifyName      (NAME_None)
+	, IntType            (EIntType::None)
+	{
+		checkSlow(Property);
+
+		EArrayType::Type ArrType         = EArrayType::None;
+		EPropertyFlags   PropagateFlags  = CPF_None;
+		FFieldClass*     ClassOfProperty = Property->GetClass();
+
+		if( ClassOfProperty==FArrayProperty::StaticClass() )
+		{
+			ArrType = EArrayType::Dynamic;
+
+			// if we're an array, save up Parm flags so we can propagate them.
+			// below the array will be assigned the inner property flags. This allows propagation of Parm flags (out, optional..)
+			PropagateFlags = Property->PropertyFlags & CPF_ParmFlags;
+			Property = CastFieldChecked<FArrayProperty>(Property)->Inner;
+			ClassOfProperty = Property->GetClass();
+		}
+
+		if( ClassOfProperty==FByteProperty::StaticClass() )
+		{
+			*this = FPropertyBase(CPT_Byte);
+			Enum = CastField<FByteProperty>(Property)->Enum;
+			IntType = EIntType::Sized;
+		}
+		else if( ClassOfProperty==FEnumProperty::StaticClass() )
+		{
+			FEnumProperty* EnumProp = CastField<FEnumProperty>(Property);
+			FNumericProperty* UnderlyingProp = EnumProp->GetUnderlyingProperty();
+
+			*this = FPropertyBase(
+				  UnderlyingProp->IsA<FInt8Property>()   ? CPT_Int8
+				: UnderlyingProp->IsA<FInt16Property>()  ? CPT_Int16
+				: UnderlyingProp->IsA<FIntProperty>()    ? CPT_Int
+				: UnderlyingProp->IsA<FInt64Property>()  ? CPT_Int64
+				: UnderlyingProp->IsA<FByteProperty>()   ? CPT_Byte
+				: UnderlyingProp->IsA<FUInt16Property>() ? CPT_UInt16
+				: UnderlyingProp->IsA<FUInt32Property>() ? CPT_UInt32
+				: UnderlyingProp->IsA<FUInt64Property>() ? CPT_UInt64
+				:                                          CPT_None
+			);
+			check(this->Type != CPT_None);
+			Enum = EnumProp->Enum;
+			IntType = EIntType::Sized;
+		}
+		else if( ClassOfProperty==FInt8Property::StaticClass() )
+		{
+			*this = FPropertyBase(CPT_Int8);
+			IntType = EIntType::Sized;
+		}
+		else if( ClassOfProperty==FInt16Property::StaticClass() )
+		{
+			*this = FPropertyBase(CPT_Int16);
+			IntType = EIntType::Sized;
+		}
+		else if( ClassOfProperty==FIntProperty::StaticClass() )
+		{
+			*this = FPropertyBase(CPT_Int);
+			IntType = EIntType::Sized;
+		}
+		else if( ClassOfProperty==FInt64Property::StaticClass() )
+		{
+			*this = FPropertyBase(CPT_Int64);
+			IntType = EIntType::Sized;
+		}
+		else if( ClassOfProperty==FUInt16Property::StaticClass() )
+		{
+			*this = FPropertyBase(CPT_UInt16);
+			IntType = EIntType::Sized;
+		}
+		else if( ClassOfProperty==FUInt32Property::StaticClass() )
+		{
+			*this = FPropertyBase(CPT_UInt32);
+			IntType = EIntType::Sized;
+		}
+		else if( ClassOfProperty==FUInt64Property::StaticClass() )
+		{
+			*this = FPropertyBase(CPT_UInt64);
+			IntType = EIntType::Sized;
+		}
+		else if( ClassOfProperty==FBoolProperty::StaticClass() )
+		{
+			FBoolProperty* BoolProperty =  CastField<FBoolProperty>(Property);
+			if (BoolProperty->IsNativeBool())
+			{
+				*this = FPropertyBase(CPT_Bool);
+			}
+			else
+			{
+				switch( BoolProperty->ElementSize )
+				{
+				case sizeof(uint8):
+					*this = FPropertyBase(CPT_Bool8);
+					break;
+				case sizeof(uint16):
+					*this = FPropertyBase(CPT_Bool16);
+					break;
+				case sizeof(uint32):
+					*this = FPropertyBase(CPT_Bool32);
+					break;
+				case sizeof(uint64):
+					*this = FPropertyBase(CPT_Bool64);
+					break;
+				}
+			}
+		}
+		else if( ClassOfProperty==FFloatProperty::StaticClass() )
+		{
+			*this = FPropertyBase(CPT_Float);
+		}
+		else if( ClassOfProperty==FDoubleProperty::StaticClass() )
+		{
+			*this = FPropertyBase(CPT_Double);
+		}
+		else if( ClassOfProperty==FClassProperty::StaticClass() )
+		{
+			*this = FPropertyBase(CPT_ObjectReference);
+			PropertyClass = CastField<FClassProperty>(Property)->PropertyClass;
+			MetaClass = CastField<FClassProperty>(Property)->MetaClass;
+		}
+		else if( ClassOfProperty==FObjectProperty::StaticClass() )
+		{
+			*this = FPropertyBase(CPT_ObjectReference);
+			PropertyClass = CastField<FObjectProperty>(Property)->PropertyClass;
+		}
+		else if( ClassOfProperty==FWeakObjectProperty::StaticClass() )
+		{
+			*this = FPropertyBase(CPT_WeakObjectReference);
+			PropertyClass = CastField<FWeakObjectProperty>(Property)->PropertyClass;
+		}
+		else if( ClassOfProperty==FLazyObjectProperty::StaticClass() )
+		{
+			*this = FPropertyBase(CPT_LazyObjectReference);
+			PropertyClass = CastField<FLazyObjectProperty>(Property)->PropertyClass;
+		}
+		else if( ClassOfProperty==FSoftClassProperty::StaticClass() )
+		{
+			*this = FPropertyBase(CPT_SoftObjectReference);
+			PropertyClass = CastField<FSoftClassProperty>(Property)->PropertyClass;
+			MetaClass = CastField<FSoftClassProperty>(Property)->MetaClass;
+		}
+		else if( ClassOfProperty==FSoftObjectProperty::StaticClass() )
+		{
+			*this = FPropertyBase(CPT_SoftObjectReference);
+			PropertyClass = CastField<FSoftObjectProperty>(Property)->PropertyClass;
+		}
+		else if( ClassOfProperty==FNameProperty::StaticClass() )
+		{
+			*this = FPropertyBase(CPT_Name);
+		}
+		else if( ClassOfProperty==FStrProperty::StaticClass() )
+		{
+			*this = FPropertyBase(CPT_String);
+		}
+		else if( ClassOfProperty==FTextProperty::StaticClass() )
+		{
+			*this = FPropertyBase(CPT_Text);
+		}
+		else if( ClassOfProperty==FStructProperty::StaticClass() )
+		{
+			*this = FPropertyBase(CPT_Struct);
+			Struct = CastField<FStructProperty>(Property)->Struct;
+		}
+		else if( ClassOfProperty==FDelegateProperty::StaticClass() )
+		{
+			*this = FPropertyBase(CPT_Delegate);
+			Function = CastField<FDelegateProperty>(Property)->SignatureFunction;
+		}
+		else if( ClassOfProperty==FMulticastDelegateProperty::StaticClass() )
+		{
+			*this = FPropertyBase(CPT_MulticastDelegate);
+			// @todo delegate: Any other setup for calling multi-cast delegates from script needed?
+			Function = CastField<FMulticastDelegateProperty>(Property)->SignatureFunction;
+		}
+		else if ( ClassOfProperty==FInterfaceProperty::StaticClass() )
+		{
+			*this = FPropertyBase(CPT_Interface);
+			PropertyClass = CastField<FInterfaceProperty>(Property)->InterfaceClass;
+		}
+		else if (ClassOfProperty == FFieldPathProperty::StaticClass())
+		{
+			*this = FPropertyBase(CPT_FieldPath);
+			PropertyPathClass = CastField<FFieldPathProperty>(Property)->PropertyClass;
+		}
+		else
+		{
+			UE_LOG(LogCompile, Fatal, TEXT("Unknown property type '%s'"), *Property->GetFullName() );
+		}
+		ArrayType            = ArrType;
+		PropertyFlags        = Property->PropertyFlags | PropagateFlags;
+		ImpliedPropertyFlags = CPF_None;
+		RefQualifier         = ERefQualifier::None;
+		PointerType          = EPointerType::None;
+	}
+
+	explicit FPropertyBase(FFieldClass* InPropertyClass, EPropertyType InType)
+		: Type(InType)
+		, ArrayType(EArrayType::None)
+		, PropertyFlags(CPF_None)
+		, ImpliedPropertyFlags(CPF_None)
+		, RefQualifier(ERefQualifier::None)
+		, PropertyExportFlags(PROPEXPORT_Public)
+		, PropertyPathClass(InPropertyClass)
+		, MetaClass(nullptr)
+		, DelegateName(NAME_None)
+		, DelegateSignatureOwnerClass(nullptr)
+		, RepNotifyName(NAME_None)
+		, PointerType(EPointerType::None)
+		, IntType(GetSizedIntTypeFromPropertyType(InType))
+	{
+	}
+	//@}
+
+	/** @name Functions */
+	//@{
+
+	/**
+	 * Returns whether this token represents an object reference
+	 */
+	bool IsObject() const
+	{
+		return Type == CPT_ObjectReference || Type == CPT_Interface || Type == CPT_WeakObjectReference || Type == CPT_LazyObjectReference || Type == CPT_SoftObjectReference;
+	}
+
+	bool IsContainer() const
+	{
+		return (ArrayType != EArrayType::None || MapKeyProp.IsValid());
+	}
+
+	/**
+	 * Determines whether this token's type is compatible with another token's type.
+	 *
+	 * @param	Other							the token to check against this one.
+	 *											Given the following example expressions, VarA is Other and VarB is 'this':
+	 *												VarA = VarB;
+	 *
+	 *												function func(type VarB) {}
+	 *												func(VarA);
+	 *
+	 *												static operator==(type VarB_1, type VarB_2) {}
+	 *												if ( VarA_1 == VarA_2 ) {}
+	 *
+	 * @param	bDisallowGeneralization			controls whether it should be considered a match if this token's type is a generalization
+	 *											of the other token's type (or vice versa, when dealing with structs
+	 * @param	bIgnoreImplementedInterfaces	controls whether two types can be considered a match if one type is an interface implemented
+	 *											by the other type.
+	 */
+	bool MatchesType( const FPropertyBase& Other, bool bDisallowGeneralization, bool bIgnoreImplementedInterfaces=false ) const
+	{
+		check(Type!=CPT_None || !bDisallowGeneralization);
+
+		bool bIsObjectType = IsObject();
+		bool bOtherIsObjectType = Other.IsObject();
+		bool bIsObjectComparison = bIsObjectType && bOtherIsObjectType;
+		bool bReverseClassChainCheck = true;
+
+		// If converting to an l-value, we require an exact match with an l-value.
+		if( (PropertyFlags&CPF_OutParm) != 0 )
+		{
+			// if the other type is not an l-value, disallow
+			if ( (Other.PropertyFlags&CPF_OutParm) == 0 )
+			{
+				return false;
+			}
+
+			// if the other type is const and we are not const, disallow
+			if ( (Other.PropertyFlags&CPF_ConstParm) != 0 && (PropertyFlags&CPF_ConstParm) == 0 )
+			{
+				return false;
+			}
+
+			if ( Type == CPT_Struct )
+			{
+				// Allow derived structs to be passed by reference, unless this is a dynamic array of structs
+				bDisallowGeneralization = bDisallowGeneralization || ArrayType == EArrayType::Dynamic || Other.ArrayType == EArrayType::Dynamic;
+			}
+
+			// if Type == CPT_ObjectReference, out object function parm; allow derived classes to be passed in
+			// if Type == CPT_Interface, out interface function parm; allow derived classes to be passed in
+			else if ( (PropertyFlags & CPF_ConstParm) == 0 || !IsObject() )
+			{
+				// all other variable types must match exactly when passed as the value to an 'out' parameter
+				bDisallowGeneralization = true;
+			}
+			
+			// both types are objects, but one is an interface and one is an object reference
+			else if ( bIsObjectComparison && Type != Other.Type )
+			{
+				return false;
+			}
+		}
+		else if ((Type == CPT_ObjectReference || Type == CPT_WeakObjectReference || Type == CPT_LazyObjectReference || Type == CPT_SoftObjectReference) && Other.Type != CPT_Interface && (PropertyFlags & CPF_ReturnParm))
+		{
+			bReverseClassChainCheck = false;
+		}
+
+		// Check everything.
+		if( Type==CPT_None && (Other.Type==CPT_None || !bDisallowGeneralization) )
+		{
+			// If Other has no type, accept anything.
+			return true;
+		}
+		else if( Type != Other.Type && !bIsObjectComparison )
+		{
+			// Mismatched base types.
+			return false;
+		}
+		else if( ArrayType != Other.ArrayType )
+		{
+			// Mismatched array types.
+			return false;
+		}
+		else if( Type==CPT_Byte )
+		{
+			// Make sure enums match, or we're generalizing.
+			return Enum==Other.Enum || (Enum==NULL && !bDisallowGeneralization);
+		}
+		else if( bIsObjectType )
+		{
+			check(PropertyClass!=NULL);
+
+			// Make sure object types match, or we're generalizing.
+			if( bDisallowGeneralization )
+			{
+				// Exact match required.
+				return PropertyClass==Other.PropertyClass && MetaClass==Other.MetaClass;
+			}
+			else if( Other.PropertyClass==NULL )
+			{
+				// Cannonical 'None' matches all object classes.
+				return true;
+			}
+			else
+			{
+				// Generalization is ok (typical example of this check would look like: VarA = VarB;, where this is VarB and Other is VarA)
+				if ( Other.PropertyClass->IsChildOf(PropertyClass) )
+				{
+					if ( !bIgnoreImplementedInterfaces || ((Type == CPT_Interface) == (Other.Type == CPT_Interface)) )
+					{
+						if ( !PropertyClass->IsChildOf(UClass::StaticClass()) || MetaClass == NULL || Other.MetaClass->IsChildOf(MetaClass) ||
+							(bReverseClassChainCheck && (Other.MetaClass == NULL || MetaClass->IsChildOf(Other.MetaClass))) )
+						{
+							return true;
+						}
+					}
+				}
+				// check the opposite class chain for object types
+				else if (bReverseClassChainCheck && Type != CPT_Interface && bIsObjectComparison && PropertyClass != NULL && PropertyClass->IsChildOf(Other.PropertyClass))
+				{
+					if (!Other.PropertyClass->IsChildOf(UClass::StaticClass()) || MetaClass == NULL || Other.MetaClass == NULL || MetaClass->IsChildOf(Other.MetaClass) || Other.MetaClass->IsChildOf(MetaClass))
+					{
+						return true;
+					}
+				}
+
+				if ( PropertyClass->HasAnyClassFlags(CLASS_Interface) && !bIgnoreImplementedInterfaces )
+				{
+					if ( Other.PropertyClass->ImplementsInterface(PropertyClass) )
+					{
+						return true;
+					}
+				}
+
+				return false;
+			}
+		}
+		else if( Type==CPT_Struct )
+		{
+			check(Struct!=NULL);
+			check(Other.Struct!=NULL);
+
+			if ( Struct == Other.Struct )
+			{
+				// struct types match exactly 
+				return true;
+			}
+
+			// returning false here prevents structs related through inheritance from being used interchangeably, such as passing a derived struct as the value for a parameter
+			// that expects the base struct, or vice versa.  An easier example is assignment (e.g. Vector = Plane or Plane = Vector).
+			// there are two cases to consider (let's use vector and plane for the example):
+			// - Vector = Plane;
+			//		in this expression, 'this' is the vector, and Other is the plane.  This is an unsafe conversion, as the destination property type is used to copy the r-value to the l-value
+			//		so in this case, the VM would call CopyCompleteValue on the FPlane struct, which would copy 16 bytes into the l-value's buffer;  However, the l-value buffer will only be
+			//		12 bytes because that is the size of FVector
+			// - Plane = Vector;
+			//		in this expression, 'this' is the plane, and Other is the vector.  This is a safe conversion, since only 12 bytes would be copied from the r-value into the l-value's buffer
+			//		(which would be 16 bytes).  The problem with allowing this conversion is that what to do with the extra member (e.g. Plane.W); should it be left alone? should it be zeroed?
+			//		difficult to say what the correct behavior should be, so let's just ignore inheritance for the sake of determining whether two structs are identical
+
+			// Previously, the logic for determining whether this is a generalization of Other was reversed; this is very likely the culprit behind all current issues with 
+			// using derived structs interchangeably with their base versions.  The inheritance check has been fixed; for now, allow struct generalization and see if we can find any further
+			// issues with allowing conversion.  If so, then we disable all struct generalization by returning false here.
+ 			// return false;
+
+			if ( bDisallowGeneralization )
+			{
+				return false;
+			}
+
+			// Generalization is ok if this is not a dynamic array
+			if ( ArrayType != EArrayType::Dynamic && Other.ArrayType != EArrayType::Dynamic )
+			{
+				if ( !Other.Struct->IsChildOf(Struct) && Struct->IsChildOf(Other.Struct) )
+				{
+					return true;
+				}
+			}
+
+			return false;
+		}
+		else
+		{
+			// General match.
+			return true;
+		}
+	}
+
+	FString Describe()
+	{
+		return FString::Printf(
+			TEXT("Type:%s  Flags:%lli  ImpliedFlags:%lli  Enum:%s  PropertyClass:%s  Struct:%s  Function:%s  MetaClass:%s"),
+			GetPropertyTypeText(Type), PropertyFlags, ImpliedPropertyFlags,
+			Enum!=NULL?*Enum->GetName():TEXT(""),
+			PropertyClass!=NULL?*PropertyClass->GetName():TEXT("NULL"),
+			Struct!=NULL?*Struct->GetName():TEXT("NULL"),
+			Function!=NULL?*Function->GetName():TEXT("NULL"),
+			MetaClass!=NULL?*MetaClass->GetName():TEXT("NULL")
+			);
+	}
+	//@}
+
+	EIntType GetSizedIntTypeFromPropertyType(EPropertyType PropType)
+	{
+		switch (PropType)
+		{
+			case CPT_Byte:
+			case CPT_UInt16:
+			case CPT_UInt32:
+			case CPT_UInt64:
+			case CPT_Int8:
+			case CPT_Int16:
+			case CPT_Int:
+			case CPT_Int64:
+				return EIntType::Sized;
+
+			default:
+				return EIntType::None;
+		}
+	}
+
+	static const TCHAR* GetPropertyTypeText( EPropertyType Type );
+
+	friend struct FPropertyBaseArchiveProxy;
+};
+
+//
+// Token types.
+//
+enum ETokenType
+{
+	TOKEN_None				= 0x00,		// No token.
+	TOKEN_Identifier		= 0x01,		// Alphanumeric identifier.
+	TOKEN_Symbol			= 0x02,		// Symbol.
+	TOKEN_Const				= 0x03,		// A constant.
+	TOKEN_Max				= 0x0D
+};
+
+/*-----------------------------------------------------------------------------
+	FToken.
+-----------------------------------------------------------------------------*/
+//
+// Information about a token that was just parsed.
+//
+class FToken : public FPropertyBase
+{
+public:
+	/** @name Variables */
+	//@{
+	/** Type of token. */
+	ETokenType TokenType;
+
+private:
+	/** Whether TokenName has been looked up */
+	mutable bool bTokenNameInitialized;
+	/** Name of token, lazily initialized */
+	mutable FName TokenName;
+
+public:
+	/** Starting position in script where this token came from. */
+	int32 StartPos;
+	/** Starting line in script. */
+	int32 StartLine;
+	/** Always valid. */
+	TCHAR Identifier[NAME_SIZE];
+	/** property that corresponds to this FToken - null if this Token doesn't correspond to a FProperty */
+	FProperty* TokenProperty;
+
+
+public:
+
+	union
+	{
+		// TOKEN_Const values.
+		uint8 Byte;								// If CPT_Byte.
+		int64 Int64;							// If CPT_Int64.
+		int32 Int;								// If CPT_Int.
+		bool NativeBool;						// if CPT_Bool
+		float Float;							// If CPT_Float.
+		double Double;							// If CPT_Double.
+		uint8 NameBytes[sizeof(FName)];			// If CPT_Name.
+		TCHAR String[MAX_STRING_CONST_SIZE];	// If CPT_String
+	};
+	//@}
+
+	FToken& operator=(const FToken& Other);
+	FToken& operator=(FToken&& Other);
+
+	FString GetConstantValue() const
+	{
+		if (TokenType == TOKEN_Const)
+		{
+			switch (Type)
+			{
+				case CPT_Byte:
+					return FString::Printf(TEXT("%u"), Byte);
+				case CPT_Int64:
+					return FString::Printf(TEXT("%" INT64_FMT), Int64);
+				case CPT_Int:
+					return FString::Printf(TEXT("%i"), Int);
+				case CPT_Bool:
+					// Don't use FCoreTexts::True/FCoreTexts::False here because they can be localized
+					return FString::Printf(TEXT("%s"), NativeBool ? *(FName::GetEntry(NAME_TRUE)->GetPlainNameString()) : *(FName::GetEntry(NAME_FALSE)->GetPlainNameString()));
+				case CPT_Float:
+					return FString::Printf(TEXT("%f"), Float);
+				case CPT_Double:
+					return FString::Printf(TEXT("%f"), Double);
+				case CPT_Name:
+					return FString::Printf(TEXT("%s"), *(*(FName*)NameBytes).ToString());
+				case CPT_String:
+					return String;
+
+				// unsupported (parsing never produces a constant token of these types)
+				default:
+					return TEXT("InvalidTypeForAToken");
+			}
+		}
+		else
+		{
+			return TEXT("NotConstant");
+		}
+	}
+
+	// Constructors.
+	FToken()
+		: FPropertyBase(CPT_None)
+		, TokenProperty(NULL)
+	{
+		InitToken(CPT_None);
+	}
+
+	// copy constructors
+	explicit FToken(EPropertyType InType)
+		: FPropertyBase(InType)
+		, TokenProperty(NULL)
+	{
+		InitToken(InType);
+	}
+
+	// copy constructors
+	FToken(const FPropertyBase& InType)
+		: FPropertyBase(CPT_None)
+		, TokenProperty(NULL)
+	{
+		InitToken( CPT_None );
+		(FPropertyBase&)*this = InType;
+	}
+
+	FToken(const FToken& InToken)
+		: FPropertyBase(CPT_None)
+	{
+		*this = InToken;
+	}
+
+	FToken(FToken&& InToken)
+		: FPropertyBase(CPT_None)
+	{
+		*this = MoveTemp(InToken);
+	}
+
+	// Inlines.
+	void InitToken( EPropertyType InType )
+	{
+		(FPropertyBase&)*this = FPropertyBase(InType);
+		TokenType		= TOKEN_None;
+		TokenName		= NAME_None;
+		StartPos		= 0;
+		StartLine		= 0;
+		*Identifier		= 0;
+		FMemory::Memzero(String);
+	}
+	bool Matches(const TCHAR Ch) const
+	{
+		return TokenType==TOKEN_Symbol && Identifier[0] == Ch && Identifier[1] == 0;
+	}
+	bool Matches( const TCHAR* Str, ESearchCase::Type SearchCase ) const
+	{
+		return (TokenType==TOKEN_Identifier || TokenType==TOKEN_Symbol) && ((SearchCase == ESearchCase::CaseSensitive) ? !FCString::Strcmp(Identifier, Str) : !FCString::Stricmp(Identifier, Str));
+	}
+	bool IsBool() const
+	{
+		return Type == CPT_Bool || Type == CPT_Bool8 || Type == CPT_Bool16 || Type == CPT_Bool32 || Type == CPT_Bool64;
+	}
+	FName GetTokenName() const
+	{
+		if (!bTokenNameInitialized)
+		{
+			TokenName = FName(Identifier, FNAME_Find);
+			bTokenNameInitialized = true;
+		}
+		return TokenName;
+	}
+	void ClearTokenName()
+	{
+		bTokenNameInitialized = false;
+		TokenName = NAME_None;
+	}
+
+	// Setters.
+	
+	void SetIdentifier( const TCHAR* InString)
+	{
+		InitToken(CPT_None);
+		TokenType = TOKEN_Identifier;
+		FCString::Strncpy(Identifier, InString, NAME_SIZE);
+		bTokenNameInitialized = false;
+	}
+
+	void SetConstInt64( int64 InInt64 )
+	{
+		(FPropertyBase&)*this = FPropertyBase(CPT_Int64);
+		Int64			= InInt64;
+		TokenType		= TOKEN_Const;
+	}
+	void SetConstInt( int32 InInt )
+	{
+		(FPropertyBase&)*this = FPropertyBase(CPT_Int);
+		Int				= InInt;
+		TokenType		= TOKEN_Const;
+	}
+	void SetConstBool( bool InBool )
+	{
+		(FPropertyBase&)*this = FPropertyBase(CPT_Bool);
+		NativeBool 		= InBool;
+		TokenType		= TOKEN_Const;
+	}
+	void SetConstFloat( float InFloat )
+	{
+		(FPropertyBase&)*this = FPropertyBase(CPT_Float);
+		Float			= InFloat;
+		TokenType		= TOKEN_Const;
+	}
+	void SetConstDouble( double InDouble )
+	{
+		(FPropertyBase&)*this = FPropertyBase(CPT_Double);
+		Double			= InDouble;
+		TokenType		= TOKEN_Const;
+	}
+	void SetConstName( FName InName )
+	{
+		(FPropertyBase&)*this = FPropertyBase(CPT_Name);
+		*(FName *)NameBytes = InName;
+		TokenType		= TOKEN_Const;
+	}
+	void SetConstString( const TCHAR* InString, int32 MaxLength=MAX_STRING_CONST_SIZE )
+	{
+		check(MaxLength>0);
+		(FPropertyBase&)*this = FPropertyBase(CPT_String);
+		if( InString != String )
+		{
+			FCString::Strncpy( String, InString, MaxLength );
+		}
+		TokenType = TOKEN_Const;
+	}
+	void SetConstChar(TCHAR InChar)
+	{
+		//@TODO: Treating this like a string for now, nothing consumes it
+		(FPropertyBase&)*this = FPropertyBase(CPT_String);
+		String[0] = InChar;
+		String[1] = 0;
+		TokenType = TOKEN_Const;
+	}
+	//!!struct constants
+
+	// Getters.
+	bool GetConstInt( int32& I ) const
+	{
+		if( TokenType==TOKEN_Const && Type==CPT_Int64 )
+		{
+			I = (int32)Int64;
+			return 1;
+		}
+		else if( TokenType==TOKEN_Const && Type==CPT_Int )
+		{
+			I = Int;
+			return 1;
+		}
+		else if( TokenType==TOKEN_Const && Type==CPT_Byte )
+		{
+			I = Byte;
+			return 1;
+		}
+		else if( TokenType==TOKEN_Const && Type==CPT_Float && Float==FMath::TruncToInt(Float))
+		{
+			I = (int32) Float;
+			return 1;
+		}
+		else if (TokenType == TOKEN_Const && Type == CPT_Double && Double == FMath::TruncToInt((float)Double))
+		{
+			I = (int32) Double;
+			return 1;
+		}
+		else return 0;
+	}
+
+	bool GetConstInt64( int64& I ) const
+	{
+		if( TokenType==TOKEN_Const && Type==CPT_Int64 )
+		{
+			I = Int64;
+			return 1;
+		}
+		else if( TokenType==TOKEN_Const && Type==CPT_Int )
+		{
+			I = Int;
+			return 1;
+		}
+		else if( TokenType==TOKEN_Const && Type==CPT_Byte )
+		{
+			I = Byte;
+			return 1;
+		}
+		else if( TokenType==TOKEN_Const && Type==CPT_Float && Float==FMath::TruncToInt(Float))
+		{
+			I = (int32) Float;
+			return 1;
+		}
+		else if (TokenType == TOKEN_Const && Type == CPT_Double && Double == FMath::TruncToInt((float)Double))
+		{
+			I = (int32) Double;
+			return 1;
+		}
+		else return 0;
+	}
+
+	FString Describe()
+	{
+		return FString::Printf(
+			TEXT("Property:%s  Type:%s  TokenName:%s  ConstValue:%s  Struct:%s  Flags:%lli  Implied:%lli"),
+			TokenProperty!=NULL?*TokenProperty->GetName():TEXT("NULL"),
+			GetPropertyTypeText(Type), *GetTokenName().ToString(), *GetConstantValue(),
+			Struct!=NULL?*Struct->GetName():TEXT("NULL"),
+			PropertyFlags,
+			ImpliedPropertyFlags
+			);
+	}
+
+	friend struct FTokenArchiveProxy;
+};
+
+/**
+ * Information about a function being compiled.
+ */
+struct FFuncInfo
+{
+	/** @name Variables */
+	//@{
+	/** Name of the function or operator. */
+	FToken		Function;
+	/** Function flags. */
+	EFunctionFlags	FunctionFlags;
+	/** Function flags which are only required for exporting */
+	uint32		FunctionExportFlags;
+	/** Number of parameters expected for operator. */
+	int32		ExpectParms;
+	/** Pointer to the UFunction corresponding to this FFuncInfo */
+	UFunction*	FunctionReference;
+	/** Name of the wrapper function that marshalls the arguments and does the indirect call **/
+	FString		MarshallAndCallName;
+	/** Name of the actual implementation **/
+	FString		CppImplName;
+	/** Name of the actual validation implementation **/
+	FString		CppValidationImplName;
+	/** Name for callback-style names **/
+	FString		UnMarshallAndCallName;
+	/** Endpoint name */
+	FString		EndpointName;
+	/** Identifier for an RPC call to a platform service */
+	uint16		RPCId;
+	/** Identifier for an RPC call expecting a response */
+	uint16		RPCResponseId;
+	/** Delegate macro line in header. */
+	int32		MacroLine;
+	/** Position in file where this function was declared. Points to first char of function name. */
+	int32 InputPos;
+	/** Whether this function represents a sealed event */
+	bool		bSealedEvent;
+	/** TRUE if the function is being forced to be considered as impure by the user */
+	bool bForceBlueprintImpure;
+
+	//@}
+
+	/** Constructor. */
+	FFuncInfo()
+		: Function()
+		, FunctionFlags(FUNC_None)
+		, FunctionExportFlags(0)
+		, ExpectParms(0)
+		, FunctionReference(nullptr)
+		, RPCId(0)
+		, RPCResponseId(0)
+		, MacroLine(-1)
+		, InputPos(-1)
+		, bSealedEvent(false)
+		, bForceBlueprintImpure(false)
+	{}
+
+	FFuncInfo(const FFuncInfo& Other) = default;
+	FFuncInfo(FFuncInfo&& Other) = default;
+	FFuncInfo& operator=(FFuncInfo&& Other) = default;
+
+	/** Set the internal function names based on flags **/
+	void SetFunctionNames()
+	{
+		FString FunctionName = FunctionReference->GetName();
+		if( FunctionReference->HasAnyFunctionFlags( FUNC_Delegate ) )
+		{
+			FunctionName.LeftChopInline( FString(HEADER_GENERATED_DELEGATE_SIGNATURE_SUFFIX).Len(), false );
+		}
+		UnMarshallAndCallName = FString(TEXT("exec")) + FunctionName;
+		
+		if (FunctionReference->HasAnyFunctionFlags(FUNC_BlueprintEvent))
+		{
+			MarshallAndCallName = FunctionName;
+		}
+		else
+		{
+			MarshallAndCallName = FString(TEXT("event")) + FunctionName;
+		}
+
+		if (FunctionReference->HasAllFunctionFlags(FUNC_Native | FUNC_Net))
+		{
+			MarshallAndCallName = FunctionName;
+			if (FunctionReference->HasAllFunctionFlags(FUNC_NetResponse))
+			{
+				// Response function implemented by programmer and called directly from thunk
+				CppImplName = FunctionReference->GetName();
+			}
+			else
+			{
+				if (CppImplName.IsEmpty())
+				{
+					CppImplName = FunctionReference->GetName() + TEXT("_Implementation");
+				}
+				else if (CppImplName == FunctionName)
+				{
+					FError::Throwf(TEXT("Native implementation function must be different than original function name."));
+				}
+
+				if (CppValidationImplName.IsEmpty() && FunctionReference->HasAllFunctionFlags(FUNC_NetValidate))
+				{
+					CppValidationImplName = FunctionReference->GetName() + TEXT("_Validate");
+				}
+				else if(CppValidationImplName == FunctionName)
+				{
+					FError::Throwf(TEXT("Validation function must be different than original function name."));
+				}
+			}
+		}
+
+		if (FunctionReference->HasAllFunctionFlags(FUNC_Delegate))
+		{
+			MarshallAndCallName = FString(TEXT("delegate")) + FunctionName;
+		}
+
+		if (FunctionReference->HasAllFunctionFlags(FUNC_BlueprintEvent | FUNC_Native))
+		{
+			MarshallAndCallName = FunctionName;
+			CppImplName = FunctionReference->GetName() + TEXT("_Implementation");
+		}
+
+		if (CppImplName.IsEmpty())
+		{
+			CppImplName = FunctionName;
+		}
+	}
+};
+
+/**
+ * Stores "compiler" data about an FToken.  "Compiler" data is data that is associated with a
+ * specific property, function or class that is only needed during script compile.
+ * This class is designed to make adding new compiler data very simple.
+ *
+ * - stores the raw evaluated bytecode associated with an FToken
+ */
+struct FTokenData
+{
+	/** The token tracked by this FTokenData. */
+	FToken Token;
+
+	FTokenData() = default;
+	FTokenData(FTokenData&&) = default;
+	FTokenData& operator=(FTokenData&&) = default;
+	FTokenData(FToken&& InToken)
+		: Token(MoveTemp(InToken))
+	{
+	}
+};
+
+/**
+ * Class for storing data about a list of properties.  Though FToken contains a reference to its
+ * associated FProperty, it's faster lookup to use the FProperty as the key in a TMap.
+ */
+class FPropertyData : public TMap< FProperty*, TSharedPtr<FTokenData> >
+{
+	typedef TMap<FProperty*, TSharedPtr<FTokenData> >	Super;
+
+public:
+	/**
+	 * Returns the value associated with a specified key.
+	 * @param	Key - The key to search for.
+	 * @return	A pointer to the value associated with the specified key, or NULL if the key isn't contained in this map.  The pointer
+	 *			is only valid until the next change to any key in the map.
+	 */
+	FTokenData* Find(FProperty* Key)
+	{
+		FTokenData* Result = NULL;
+
+		TSharedPtr<FTokenData>* pResult = Super::Find(Key);
+		if ( pResult != NULL )
+		{
+			Result = pResult->Get();
+		}
+		return Result;
+	}
+	const FTokenData* Find(FProperty* Key) const
+	{
+		const FTokenData* Result = NULL;
+
+		const TSharedPtr<FTokenData>* pResult = Super::Find(Key);
+		if ( pResult != NULL )
+		{
+			Result = pResult->Get();
+		}
+		return Result;
+	}
+
+	/**
+	 * Sets the value associated with a key.  If the key already exists in the map, uses the same
+	 * value pointer and reinitalized the FTokenData with the input value.
+	 *
+	 * @param	InKey	the property to get a token wrapper for
+	 * @param	InValue	the token wrapper for the specified key
+	 *
+	 * @return	a pointer to token data created associated with the property
+	 */
+	FTokenData* Set(FProperty* InKey, FTokenData&& InValue, FUnrealSourceFile* UnrealSourceFile);
+
+	/**
+	 * (debug) Dumps the values of this FPropertyData to the log file
+	 * 
+	 * @param	Indent	number of spaces to insert at the beginning of each line
+	 */	
+	void Dump( int32 Indent )
+	{
+		for (auto& Kvp : *this)
+		{
+			TSharedPtr<FTokenData>& PointerVal = Kvp.Value;
+			FToken& Token = PointerVal->Token;
+			if ( Token.Type != CPT_None )
+			{
+				UE_LOG(LogCompile, Log, TEXT("%s%s"), FCString::Spc(Indent), *Token.Describe());
+			}
+		}
+	}
+};
+
+/**
+ * Class for storing additional data about compiled structs and struct properties
+ */
+class FStructData
+{
+public:
+	/** info about the struct itself */
+	FToken			StructData;
+
+private:
+	/** info for the properties contained in this struct */
+	FPropertyData	StructPropertyData;
+
+public:
+	/**
+	 * Adds a new struct property token
+	 * 
+	 * @param	PropertyToken	token that should be added to the list
+	 */
+	void AddStructProperty(FTokenData&& PropertyToken, FUnrealSourceFile* UnrealSourceFile)
+	{
+		check(PropertyToken.Token.TokenProperty);
+		StructPropertyData.Set(PropertyToken.Token.TokenProperty, MoveTemp(PropertyToken), UnrealSourceFile);
+	}
+
+	FPropertyData& GetStructPropertyData()
+	{
+		return StructPropertyData;
+	}
+	const FPropertyData& GetStructPropertyData() const
+	{
+		return StructPropertyData;
+	}
+
+	/**
+	* (debug) Dumps the values of this FStructData to the log file
+	* 
+	* @param	Indent	number of spaces to insert at the beginning of each line
+	*/	
+	void Dump( int32 Indent )
+	{
+		UE_LOG(LogCompile, Log, TEXT("%s%s"), FCString::Spc(Indent), *StructData.Describe());
+
+		UE_LOG(LogCompile, Log, TEXT("%sproperties:"), FCString::Spc(Indent));
+		StructPropertyData.Dump(Indent + 4);
+	}
+
+	/** Constructor */
+	FStructData(FToken&& StructToken) : StructData(MoveTemp(StructToken)) {}
+
+	friend struct FStructDataArchiveProxy;
+};
+
+/**
+ * Class for storing additional data about compiled function properties.
+ */
+class FFunctionData
+{
+	/** info about the function associated with this FFunctionData */
+	FFuncInfo		FunctionData;
+
+	/** return value for this function */
+	FTokenData		ReturnTypeData;
+
+	/** function parameter data */
+	FPropertyData	ParameterData;
+
+	/**
+	 * Adds a new parameter token
+	 * 
+	 * @param	PropertyToken	token that should be added to the list
+	 */
+	void AddParameter(FToken&& PropertyToken, FUnrealSourceFile* UnrealSourceFile)
+	{
+		check(PropertyToken.TokenProperty);
+		ParameterData.Set(PropertyToken.TokenProperty, MoveTemp(PropertyToken), UnrealSourceFile);
+	}
+
+	/**
+	 * Sets the value of the return token for this function
+	 * 
+	 * @param	PropertyToken	token that should be added
+	 */
+	void SetReturnData(FToken&& PropertyToken )
+	{
+		check(PropertyToken.TokenProperty);
+		ReturnTypeData.Token = MoveTemp(PropertyToken);
+	}
+
+public:
+	/** Constructors */
+	FFunctionData() = default;
+	FFunctionData(FFunctionData&& Other) = default;
+	FFunctionData(FFuncInfo&& inFunctionData )
+		: FunctionData(MoveTemp(inFunctionData))
+	{}
+
+	FFunctionData& operator=(FFunctionData&& Other) = default;
+	
+	/** @name getters */
+	//@{
+	const	FFuncInfo&		GetFunctionData()	const	{	return FunctionData;	}
+	const	FToken&			GetReturnData()		const	{	return ReturnTypeData.Token;	}
+	const	FPropertyData&	GetParameterData()	const	{	return ParameterData;	}
+	FPropertyData&			GetParameterData()			{	return ParameterData;	}
+	FTokenData* GetReturnTokenData() { return &ReturnTypeData; }
+	//@}
+
+	void UpdateFunctionData(FFuncInfo& UpdatedFuncData)
+	{
+		//@TODO: UCREMOVAL: Some more thorough evaluation should be done here
+		FunctionData.FunctionFlags |= UpdatedFuncData.FunctionFlags;
+		FunctionData.FunctionExportFlags |= UpdatedFuncData.FunctionExportFlags;
+	}
+
+	/**
+	 * Adds a new function property to be tracked.  Determines whether the property is a
+	 * function parameter, local property, or return value, and adds it to the appropriate
+	 * list
+	 * 
+	 * @param	PropertyToken	the property to add
+	 */
+	void AddProperty(FToken&& PropertyToken, FUnrealSourceFile* UnrealSourceFile)
+	{
+		const FProperty* Prop = PropertyToken.TokenProperty;
+		check(Prop);
+		check( (Prop->PropertyFlags&CPF_Parm) != 0 );
+
+		if ( (Prop->PropertyFlags&CPF_ReturnParm) != 0 )
+		{
+			SetReturnData(MoveTemp(PropertyToken));
+		}
+		else
+		{
+			AddParameter(MoveTemp(PropertyToken), UnrealSourceFile);
+		}
+	}
+
+	/**
+	 * (debug) Dumps the values of this FFunctionData to the log file
+	 * 
+	 * @param	Indent	number of spaces to insert at the beginning of each line
+	 */	
+	void Dump( int32 Indent )
+	{
+		UE_LOG(LogCompile, Log, TEXT("%sparameters:"), FCString::Spc(Indent));
+		ParameterData.Dump(Indent + 4);
+
+		UE_LOG(LogCompile, Log, TEXT("%sreturn prop:"), FCString::Spc(Indent));
+		if ( ReturnTypeData.Token.Type != CPT_None )
+		{
+			UE_LOG(LogCompile, Log, TEXT("%s%s"), FCString::Spc(Indent + 4), *ReturnTypeData.Token.Describe());
+		}
+	}
+
+	/**
+	 * Sets the specified function export flags
+	 */
+	void SetFunctionExportFlag( uint32 NewFlags )
+	{
+		FunctionData.FunctionExportFlags |= NewFlags;
+	}
+
+	/**
+	 * Clears the specified function export flags
+	 */
+	void ClearFunctionExportFlags( uint32 ClearFlags )
+	{
+		FunctionData.FunctionExportFlags &= ~ClearFlags;
+	}
+
+	/**
+	 * Finds function data for given function object.
+	 */
+	static FFunctionData* FindForFunction(UFunction* Function);
+
+	/**
+	 * Adds function data object for given function object.
+	 */
+	static FFunctionData* Add(UFunction* Function);
+
+	/**
+	 * Adds function data object for given function object.
+	 */
+	static FFunctionData* Add(FFuncInfo&& FunctionInfo);
+
+	/**
+	 * Tries to find function data for given function object.
+	 */
+	static bool TryFindForFunction(UFunction* Function, FFunctionData*& OutData);
+
+private:
+	static TMap<UFunction*, TUniqueObj<FFunctionData> > FunctionDataMap;
+};
+
+/**
+ * Tracks information about a multiple inheritance parent declaration for native script classes.
+ */
+struct FMultipleInheritanceBaseClass
+{
+	/**
+	 * The name to use for the base class when exporting the script class to header file.
+	 */
+	FString ClassName;
+
+	/**
+	 * For multiple inheritance parents declared using 'Implements', corresponds to the UClass for the interface.  For multiple inheritance parents declared
+	 * using 'Inherits', this value will be NULL.
+	 */
+	UClass* InterfaceClass = nullptr;
+
+	/**
+	 * Constructors
+	 */
+	explicit FMultipleInheritanceBaseClass(FString&& BaseClassName)
+		: ClassName(MoveTemp(BaseClassName))
+	{
+	}
+
+	explicit FMultipleInheritanceBaseClass(UClass* ImplementedInterfaceClass)
+		: InterfaceClass(ImplementedInterfaceClass)
+	{
+		ClassName = FString::Printf(TEXT("I%s"), *ImplementedInterfaceClass->GetName());
+	}
+};
+
+enum class EParsedInterface
+{
+	NotAnInterface,
+	ParsedUInterface,
+	ParsedIInterface
+};
+
+/**
+ * Class for storing compiler metadata about a class's properties.
+ */
+class FClassMetaData
+{
+	/** member properties for this class */
+	FPropertyData											GlobalPropertyData;
+
+	/** base classes to multiply inherit from (other than the main base class */
+	TArray<FMultipleInheritanceBaseClass*>					MultipleInheritanceParents;
+
+	/** whether this class declares delegate functions or properties */
+	bool													bContainsDelegates;
+
+	/** The line of UCLASS/UINTERFACE macro in this class. */
+	int32 PrologLine;
+
+	/** The line of GENERATED_BODY/GENERATED_UCLASS_BODY macro in this class. */
+	int32 GeneratedBodyLine;
+
+	/** Same as above, but for interface class associated with this class. */
+	int32 InterfaceGeneratedBodyLine;
+
+public:
+	/** Default constructor */
+	FClassMetaData()
+		: bContainsDelegates(false)
+		, PrologLine(-1)
+		, GeneratedBodyLine(-1)
+		, InterfaceGeneratedBodyLine(-1)
+		, bConstructorDeclared(false)
+		, bDefaultConstructorDeclared(false)
+		, bObjectInitializerConstructorDeclared(false)
+		, bCustomVTableHelperConstructorDeclared(false)
+		, GeneratedBodyMacroAccessSpecifier(ACCESS_NotAnAccessSpecifier)
+	{
+	}
+
+	/**
+	 * Gets prolog line number for this class.
+	 */
+	int32 GetPrologLine() const
+	{
+		check(PrologLine > 0);
+		return PrologLine;
+	}
+
+	/**
+	 * Gets generated body line number for this class.
+	 */
+	int32 GetGeneratedBodyLine() const
+	{
+		check(GeneratedBodyLine > 0);
+		return GeneratedBodyLine;
+	}
+
+	/**
+	 * Gets interface generated body line number for this class.
+	 */
+	int32 GetInterfaceGeneratedBodyLine() const
+	{
+		check(InterfaceGeneratedBodyLine > 0);
+		return InterfaceGeneratedBodyLine;
+	}
+
+	/**
+	 * Sets prolog line number for this class.
+	 */
+	void SetPrologLine(int32 Line)
+	{
+		check(Line > 0);
+		PrologLine = Line;
+	}
+
+	/**
+	 * Sets generated body line number for this class.
+	 */
+	void SetGeneratedBodyLine(int32 Line)
+	{
+		check(Line > 0);
+		GeneratedBodyLine = Line;
+	}
+
+	/**
+	 * Sets interface generated body line number for this class.
+	 */
+	void SetInterfaceGeneratedBodyLine(int32 Line)
+	{
+		check(Line > 0);
+		InterfaceGeneratedBodyLine = Line;
+	}
+
+	/**
+	 * Sets contains delegates flag for this class.
+	 */
+	void MarkContainsDelegate()
+	{
+		bContainsDelegates = true;
+	}
+
+	/**
+	 * Adds a new property to be tracked.  Determines the correct list for the property based on
+	 * its owner (function, struct, etc).
+	 * 
+	 * @param	PropertyToken	the property to add
+	 */
+	void AddProperty(FToken&& PropertyToken, FUnrealSourceFile* UnrealSourceFile)
+	{
+		FProperty* Prop = PropertyToken.TokenProperty;
+		check(Prop);
+
+		UObject* Outer = Prop->GetOwner<UObject>();
+		check(Outer);
+		UStruct* OuterClass = Cast<UStruct>(Outer);
+		if ( OuterClass != NULL )
+		{
+			// global property
+			GlobalPropertyData.Set(Prop, MoveTemp(PropertyToken), UnrealSourceFile);
+		}
+		else
+		{
+			checkNoEntry();
+			UFunction* OuterFunction = Cast<UFunction>(Outer);
+			if ( OuterFunction != NULL )
+			{
+				// function parameter, return, or local property
+				FFunctionData::FindForFunction(OuterFunction)->AddProperty(MoveTemp(PropertyToken), UnrealSourceFile);
+			}
+		}
+
+		// update the optimization flags
+		if ( !bContainsDelegates )
+		{
+			if( Prop->IsA( FDelegateProperty::StaticClass() ) || Prop->IsA( FMulticastDelegateProperty::StaticClass() ) )
+			{
+				bContainsDelegates = true;
+			}
+			else
+			{
+				FArrayProperty* ArrayProp = CastField<FArrayProperty>(Prop);
+				if ( ArrayProp != NULL )
+				{
+					if( ArrayProp->Inner->IsA( FDelegateProperty::StaticClass() ) || ArrayProp->Inner->IsA( FMulticastDelegateProperty::StaticClass() ) )
+					{
+						bContainsDelegates = true;
+					}
+				}
+			}
+		}
+	}
+
+	/**
+	 * Adds new editor-only metadata (key/value pairs) to the class or struct that
+	 * owns this property or function.
+	 * 
+	 * @param	Field		the property or function to add to
+	 * @param	InMetaData	the metadata to add
+	 */
+	static void AddMetaData(UField* Field, TMap<FName, FString>&& InMetaData)
+	{
+		// only add if we have some!
+		if (InMetaData.Num())
+		{
+			check(Field);
+
+			// get (or create) a metadata object for this package
+			UMetaData* MetaData = Field->GetOutermost()->GetMetaData();
+			TMap<FName, FString>* ExistingMetaData = MetaData->GetMapForObject(Field);
+			if (ExistingMetaData && ExistingMetaData->Num())
+			{
+				// Merge the existing metadata
+				TMap<FName, FString> MergedMetaData;
+				MergedMetaData.Reserve(InMetaData.Num() + ExistingMetaData->Num());
+				MergedMetaData.Append(*ExistingMetaData);
+				MergedMetaData.Append(InMetaData);
+				MetaData->SetObjectValues(Field, MoveTemp(MergedMetaData));
+			}
+			else
+			{
+				// set the metadata for this field
+				MetaData->SetObjectValues(Field, MoveTemp(InMetaData));
+			}
+		}
+	}
+	static void AddMetaData(FField* Field, TMap<FName, FString>&& InMetaData)
+	{
+		// only add if we have some!
+		if (InMetaData.Num())
+		{
+			check(Field);
+
+			UPackage* Package = Field->GetOutermost();
+			// get (or create) a metadata object for this package
+			UMetaData* MetaData = Package->GetMetaData();
+			
+			for (TPair<FName, FString>& MetaKeyValue : InMetaData)
+			{
+				Field->SetMetaData(MetaKeyValue.Key, MoveTemp(MetaKeyValue.Value));
+			}
+		}
+	}
+
+	/**
+	 * Finds the metadata for the function specified
+	 * 
+	 * @param	Func	the function to search for
+	 *
+	 * @return	pointer to the metadata for the function specified, or NULL
+	 *			if the function doesn't exist in the list (for example, if it
+	 *			is declared in a package that is already compiled and has had its
+	 *			source stripped)
+	 */
+	FFunctionData* FindFunctionData( UFunction* Func );
+
+	/**
+	 * Finds the metadata for the property specified
+	 * 
+	 * @param	Prop	the property to search for
+	 *
+	 * @return	pointer to the metadata for the property specified, or NULL
+	 *			if the property doesn't exist in the list (for example, if it
+	 *			is declared in a package that is already compiled and has had its
+	 *			source stripped)
+	 */
+	FTokenData* FindTokenData( FProperty* Prop );
+
+	/**
+	 * (debug) Dumps the values of this FFunctionData to the log file
+	 * 
+	 * @param	Indent	number of spaces to insert at the beginning of each line
+	 */	
+	void Dump( int32 Indent );
+
+	/**
+	 * Add a string to the list of inheritance parents for this class.
+	 *
+	 * @param Inparent The C++ class name to add to the multiple inheritance list
+	 * @param UnrealSourceFile Currently parsed source file.
+	 */
+	void AddInheritanceParent(FString&& InParent, FUnrealSourceFile* UnrealSourceFile);
+
+	/**
+	 * Add a string to the list of inheritance parents for this class.
+	 *
+	 * @param Inparent	The C++ class name to add to the multiple inheritance list
+	 * @param UnrealSourceFile Currently parsed source file.
+	 */
+	void AddInheritanceParent(UClass* ImplementedInterfaceClass, FUnrealSourceFile* UnrealSourceFile);
+
+	/**
+	 * Return the list of inheritance parents
+	 */
+	const TArray<FMultipleInheritanceBaseClass*>& GetInheritanceParents() const
+	{
+		return MultipleInheritanceParents;
+	}
+
+	/**
+	 * Returns whether this class contains any delegate properties which need to be fixed up.
+	 */
+	bool ContainsDelegates() const
+	{
+		return bContainsDelegates;
+	}
+
+	/**
+	 * Shrink TMaps to avoid slack in Pairs array.
+	 */
+	void Shrink()
+	{
+		GlobalPropertyData.Shrink();
+		MultipleInheritanceParents.Shrink();
+	}
+
+	// Is constructor declared?
+	bool bConstructorDeclared;
+
+	// Is default constructor declared?
+	bool bDefaultConstructorDeclared;
+
+	// Is ObjectInitializer constructor (i.e. a constructor with only one parameter of type FObjectInitializer) declared?
+	bool bObjectInitializerConstructorDeclared;
+
+	// Is custom VTable helper constructor declared?
+	bool bCustomVTableHelperConstructorDeclared;
+
+	// GENERATED_BODY access specifier to preserve.
+	EAccessSpecifier GeneratedBodyMacroAccessSpecifier;
+
+	/** Parsed interface state */
+	EParsedInterface ParsedInterface = EParsedInterface::NotAnInterface;
+
+	friend struct FClassMetaDataArchiveProxy;
+};
+
+/**
+ * Class for storing and linking data about properties and functions that is only required by the compiler.
+ * The type of data tracked by this class is data that would otherwise only be accessible by adding a 
+ * member property to UFunction/FProperty.  
+ */
+class FCompilerMetadataManager : protected TMap<UStruct*, TUniquePtr<FClassMetaData> >
+{
+	using Super = TMap<UStruct*, TUniquePtr<FClassMetaData>>;
+
+public:
+	/**
+	 * Adds a new class to be tracked
+	 * 
+	 * @param	Struct	the UStruct to add
+	 *
+	 * @return	a pointer to the newly added metadata for the class specified
+	 */
+	FClassMetaData* AddClassData(UStruct* Struct, FUnrealSourceFile* UnrealSourceFile);
+
+	FClassMetaData* AddInterfaceClassData(UStruct* Struct, FUnrealSourceFile* UnrealSourceFile);
+
+	/**
+	 * Find the metadata associated with the class specified
+	 * 
+	 * @param	Struct	the UStruct to add
+	 *
+	 * @return	a pointer to the newly added metadata for the class specified
+	 */
+	FClassMetaData* FindClassData(UStruct* Struct)
+	{
+		FClassMetaData* Result = nullptr;
+
+		TUniquePtr<FClassMetaData>* pClassData = Find(Struct);
+		if (pClassData)
+		{
+			Result = pClassData->Get();
+		}
+
+		return Result;
+	}
+
+	/**
+	 * Shrink TMaps to avoid slack in Pairs array.
+	 */
+	void Shrink()
+	{
+		TMap<UStruct*, TUniquePtr<FClassMetaData> >::Shrink();
+		for (TMap<UStruct*, TUniquePtr<FClassMetaData> >::TIterator It(*this); It; ++It)
+		{
+			FClassMetaData* MetaData = It->Value.Get();
+			MetaData->Shrink();
+		}
+	}
+
+	/**
+	 * Throws an exception if a UInterface was parsed but not the corresponding IInterface.
+	 */
+	void CheckForNoIInterfaces();
+
+	friend struct FCompilerMetadataManagerArchiveProxy;
+
+private:
+	TArray<TPair<UStruct*, FClassMetaData*>> InterfacesToVerify;
+};
+
+/*-----------------------------------------------------------------------------
+	Retry points.
+-----------------------------------------------------------------------------*/
+
+/**
+ * A point in the header parsing state that can be set and returned to
+ * using InitScriptLocation() and ReturnToLocation().  This is used in cases such as testing
+ * to see which overridden operator should be used, where code must be compiled
+ * and then "undone" if it was found not to match.
+ * <p>
+ * Retries are not allowed to cross command boundaries (and thus nesting 
+ * boundaries).  Retries can occur across a single command or expressions and
+ * subexpressions within a command.
+ */
+struct FScriptLocation
+{
+	static class FHeaderParser* Compiler;
+
+	/** the text buffer for the class associated with this retry point */
+	const TCHAR* Input;
+
+	/** the position into the Input buffer where this retry point is located */
+	int32 InputPos;
+
+	/** the LineNumber of the compiler when this retry point was created */
+	int32 InputLine;
+
+	/** Constructor */
+	FScriptLocation();
+};
+
+/////////////////////////////////////////////////////
+// FNameLookupCPP
+
+/**
+ * Helper class used to cache UClass* -> TCHAR* name lookup for finding the named used for C++ declaration.
+ */
+struct FNameLookupCPP
+{
+	/**
+	 * Returns the name used for declaring the passed in struct in C++
+	 *
+	 * @param	Struct	UStruct to obtain C++ name for
+	 * @return	Name used for C++ declaration
+	 */
+	static FString GetNameCPP(UStruct* Struct, bool bForceInterface = false)
+	{
+		return FString::Printf(TEXT("%s%s"), (bForceInterface ? TEXT("I") : Struct->GetPrefixCPP()), *Struct->GetName());
+	}
+};
+
+/////////////////////////////////////////////////////
+// FAdvancedDisplayParameterHandler - used by FHeaderParser::ParseParameterList, to check if a property if a function parameter has 'AdvancedDisplay' flag
+
+/** 
+ *	AdvancedDisplay can be used in two ways:
+ *	1. 'AdvancedDisplay = "3"' - the number tells how many parameters (from beginning) should NOT BE marked
+ *	2. 'AdvancedDisplay = "AttachPointName, Location, LocationType"' - list the parameters, that should BE marked
+ */
+class FAdvancedDisplayParameterHandler
+{
+	TArray<FString> ParametersNames;
+
+	int32 NumberLeaveUnmarked;
+	int32 AlreadyLeft;
+
+	bool bUseNumber;
+public:
+	FAdvancedDisplayParameterHandler(const TMap<FName, FString>* MetaData);
+
+	/** 
+	 * return if given parameter should be marked as Advance View, 
+	 * the function should be called only once for any parameter
+	 */
+	bool ShouldMarkParameter(const FString& ParameterName);
+
+	/** return if more parameters can be marked */
+	bool CanMarkMore() const;
+};
diff --git a/Engine/Source/Programs/UHTLite/Private/Scope.cpp b/Engine/Source/Programs/UHTLite/Private/Scope.cpp
new file mode 100644
--- /dev/null
+++ b/Engine/Source/Programs/UHTLite/Private/Scope.cpp
@@ -0,0 +1,196 @@
+// Copyright Epic Games, Inc. All Rights Reserved.
+
+#include "Scope.h"
+#include "UnrealHeaderTool.h"
+#include "UObject/ErrorException.h"
+#include "ParserHelper.h"
+
+extern FCompilerMetadataManager GScriptHelper;
+
+FScope::FScope(FScope* InParent)
+	: Parent(InParent)
+{ }
+
+FScope::FScope()
+	: Parent(nullptr)
+{
+
+}
+
+void FScope::AddType(UField* Type)
+{
+	TypeMap.Add(Type->GetFName(), Type);
+}
+
+/**
+ * Dispatch type to one of three arrays Enums, Structs and DelegateFunctions.
+ *
+ * @param Type Input type.
+ * @param Enums (Output parameter) Array to fill with enums.
+ * @param Structs (Output parameter) Array to fill with structs.
+ * @param DelegateFunctions (Output parameter) Array to fill with delegate functions.
+ */
+void DispatchType(UField* Type, TArray<UEnum*> &Enums, TArray<UScriptStruct*> &Structs, TArray<UDelegateFunction*> &DelegateFunctions)
+{
+	UClass* TypeClass = Type->GetClass();
+
+	if (TypeClass == UClass::StaticClass() || TypeClass == UStruct::StaticClass())
+	{
+		// Inner scopes.
+		FScope::GetTypeScope((UStruct*)Type)->SplitTypesIntoArrays(Enums, Structs, DelegateFunctions);
+	}
+	else if (TypeClass == UEnum::StaticClass())
+	{
+		UEnum* Enum = (UEnum*)Type;
+		Enums.Add(Enum);
+	}
+	else if (TypeClass == UScriptStruct::StaticClass())
+	{
+		UScriptStruct* Struct = (UScriptStruct*)Type;
+		Structs.Add(Struct);
+	}
+	else if (TypeClass == UDelegateFunction::StaticClass() || TypeClass == USparseDelegateFunction::StaticClass())
+	{
+		bool bAdded = false;
+		UDelegateFunction* Function = (UDelegateFunction*)Type;
+
+		if (Function->GetSuperFunction() == NULL)
+		{
+			DelegateFunctions.Add(Function);
+			bAdded = true;
+		}
+
+		check(bAdded);
+	}
+}
+
+void FScope::SplitTypesIntoArrays(TArray<UEnum*>& Enums, TArray<UScriptStruct*>& Structs, TArray<UDelegateFunction*>& DelegateFunctions)
+{
+	for (TPair<FName, UField*>& TypePair : TypeMap)
+	{
+		UField* Type = TypePair.Value;
+		DispatchType(Type, Enums, Structs, DelegateFunctions);
+	}
+}
+
+TSharedRef<FScope> FScope::GetTypeScope(UStruct* Type)
+{
+	auto* ScopeRefPtr = ScopeMap.Find(Type);
+	if (!ScopeRefPtr)
+	{
+		FError::Throwf(TEXT("Couldn't find scope for the type %s."), *Type->GetName());
+	}
+
+	return *ScopeRefPtr;
+}
+
+TSharedRef<FScope> FScope::AddTypeScope(UStruct* Type, FScope* ParentScope)
+{
+	FStructScope* ScopePtr = new FStructScope(Type, ParentScope);
+	TSharedRef<FScope> Scope = MakeShareable(ScopePtr);
+
+	ScopeMap.Add(Type, Scope);
+
+	return Scope;
+}
+
+UField* FScope::FindTypeByName(FName Name)
+{
+	if (!Name.IsNone())
+	{
+		TDeepScopeTypeIterator<UField, false> TypeIterator(this);
+
+		while (TypeIterator.MoveNext())
+		{
+			UField* Type = *TypeIterator;
+			if (Type->GetFName() == Name)
+			{
+				return Type;
+			}
+		}
+	}
+
+	return nullptr;
+}
+
+const UField* FScope::FindTypeByName(FName Name) const
+{
+	if (!Name.IsNone())
+	{
+		TScopeTypeIterator<UField, true> TypeIterator = GetTypeIterator();
+
+		while (TypeIterator.MoveNext())
+		{
+			UField* Type = *TypeIterator;
+			if (Type->GetFName() == Name)
+			{
+				return Type;
+			}
+		}
+	}
+
+	return nullptr;
+}
+
+bool FScope::ContainsType(UField* Type)
+{
+	return FindTypeByName(Type->GetFName()) != nullptr;
+}
+
+bool FScope::IsFileScope() const
+{
+	return Parent == nullptr;
+}
+
+bool FScope::ContainsTypes() const
+{
+	return TypeMap.Num() > 0;
+}
+
+FFileScope* FScope::GetFileScope()
+{
+	FScope* CurrentScope = this;
+	while (!CurrentScope->IsFileScope())
+	{
+		CurrentScope = const_cast<FScope*>(CurrentScope->GetParent());
+	}
+
+	return CurrentScope->AsFileScope();
+}
+
+TMap<UStruct*, TSharedRef<FScope> > FScope::ScopeMap;
+
+FFileScope::FFileScope(FName InName, FUnrealSourceFile* InSourceFile)
+	: SourceFile(InSourceFile), Name(InName)
+{ }
+
+void FFileScope::IncludeScope(FFileScope* IncludedScope)
+{
+	IncludedScopes.Add(IncludedScope);
+}
+
+FUnrealSourceFile* FFileScope::GetSourceFile() const
+{
+	return SourceFile;
+}
+
+FName FFileScope::GetName() const
+{
+	return Name;
+}
+
+UStruct* FStructScope::GetStruct() const
+{
+	return Struct;
+}
+
+FName FStructScope::GetName() const
+{
+	return Struct->GetFName();
+}
+
+FStructScope::FStructScope(UStruct* InStruct, FScope* InParent)
+	: FScope(InParent), Struct(InStruct)
+{
+
+}
diff --git a/Engine/Source/Programs/UHTLite/Private/Scope.h b/Engine/Source/Programs/UHTLite/Private/Scope.h
new file mode 100644
--- /dev/null
+++ b/Engine/Source/Programs/UHTLite/Private/Scope.h
@@ -0,0 +1,517 @@
+// Copyright Epic Games, Inc. All Rights Reserved.
+
+#pragma once
+
+#include "CoreMinimal.h"
+#include "UObject/Class.h"
+
+class UEnum;
+class UScriptStruct;
+class UDelegateFunction;
+
+// Forward declarations.
+class UStruct;
+class FClassMetaData;
+class FUnrealSourceFile;
+class FFileScope;
+class FStructScope;
+
+// Traits to achieve conditional types for const/non-const iterators.
+template <bool Condition, class TypeIfTrue, class TypeIfFalse>
+class TConditionalType
+{
+public:
+	typedef TypeIfFalse Type;
+};
+
+template <class TypeIfTrue, class TypeIfFalse>
+class TConditionalType<true, TypeIfTrue, TypeIfFalse>
+{
+public:
+	typedef TypeIfTrue Type;
+};
+
+// Base class representing type scope.
+class FScope : public TSharedFromThis<FScope>
+{
+public:
+	// Default constructor i.e. Parent == nullptr
+	FScope();
+
+	// Constructor that's providing parent scope.
+	FScope(FScope* Parent);
+
+	// Virtual destructor.
+	virtual ~FScope()
+	{ };
+
+	virtual FFileScope* AsFileScope() { return nullptr; }
+	virtual FStructScope* AsStructScope() { return nullptr; }
+
+	/**
+	 * Adds type to the scope.
+	 *
+	 * @param Type Type to add.
+	 */
+	void AddType(UField* Type);
+
+	/**
+	 * Finds type by its name.
+	 *
+	 * @param Name Name to look for.
+	 *
+	 * @returns Found type or nullptr on failure.
+	 */
+	UField* FindTypeByName(FName Name);
+
+	/**
+	 * Finds type by its name.
+	 *
+	 * @param Name Name to look for.
+	 *
+	 * @returns Found type or nullptr on failure.
+	 */
+	template <typename CharType>
+	UField* FindTypeByName(const CharType* Name)
+	{
+		return FindTypeByName(FName(Name, FNAME_Find));
+	}
+
+	/**
+	 * Finds type by its name.
+	 *
+	 * Const version.
+	 *
+	 * @param Name Name to look for.
+	 *
+	 * @returns Found type or nullptr on failure.
+	 */
+	const UField* FindTypeByName(FName Name) const;
+
+	/**
+	 * Finds type by its name.
+	 *
+	 * Const version.
+	 *
+	 * @param Name Name to look for.
+	 *
+	 * @returns Found type or nullptr on failure.
+	 */
+	template <typename CharType>
+	const UField* FindTypeByName(const CharType* Name) const
+	{
+		return FindTypeByName(FName(Name, FNAME_Find));
+	}
+	/**
+	 * Checks if scope contains this type.
+	 *
+	 * @param Type Type to look for.
+	 */
+	bool ContainsType(UField* Type);
+
+	/**
+	 * Checks if scope contains type that satisfies a predicate.
+	 *
+	 * @param Predicate Predicate to satisfy.
+	 */
+	template <typename TPredicateType>
+	bool Contains(TPredicateType Predicate)
+	{
+		for (const auto& NameTypePair : TypeMap)
+		{
+			if (Predicate.operator()(NameTypePair.Value))
+			{
+				return true;
+			}
+		}
+
+		return false;
+	}
+
+	/**
+	 * Gets struct's or class's scope.
+	 *
+	 * @param Type A struct or class for which to return the scope.
+	 *
+	 * @returns Struct's or class's scope.
+	 */
+	static TSharedRef<FScope> GetTypeScope(UStruct* Type);
+
+	/**
+	 * Adds struct's or class's scope.
+	 *
+	 * @param Type A struct or class for which to add the scope.
+	 * @param ParentScope A scope in which this scope is contained.
+	 *
+	 * @returns Newly added scope.
+	 */
+	static TSharedRef<FScope> AddTypeScope(UStruct* Type, FScope* ParentScope);
+
+	/**
+	 * Gets structs, enums and delegate functions from this scope.
+	 *
+	 * @param Enums (Output parameter) List of enums from this scope.
+	 * @param Structs (Output parameter) List of structs from this scope.
+	 * @param DelegateFunctions (Output parameter) List of delegate properties from this scope.
+	 */
+	void SplitTypesIntoArrays(TArray<UEnum*>& Enums, TArray<UScriptStruct*>& Structs, TArray<UDelegateFunction*>& DelegateFunctions);
+
+	/**
+	 * Gets scope name.
+	 */
+	virtual FName GetName() const = 0;
+
+	/**
+	 * Gets scope types filtered by type.
+	 *
+	 * @param OutArray (Output parameter) Array to fill with types from this scope.
+	 */
+	template <class TTypeFilter>
+	void GetTypes(TArray<TTypeFilter*> OutArray)
+	{
+		for (const auto& NameTypePair : TypeMap)
+		{
+			if (NameTypePair.Value->IsA<TTypeFilter>())
+			{
+				OutArray.Add((TTypeFilter*)NameTypePair.Value);
+			}
+		}
+	}
+
+	/**
+	 * Gets scope's parent.
+	 */
+	const FScope* GetParent() const
+	{
+		return Parent;
+	}
+
+	/**
+	 * Scope type iterator.
+	 */
+	template <class TType, bool TIsConst>
+	class TScopeTypeIterator
+	{
+	public:
+		// Conditional typedefs.
+		typedef typename TConditionalType<TIsConst, TMap<FName, UField*>::TConstIterator, TMap<FName, UField*>::TIterator>::Type MapIteratorType;
+		typedef typename TConditionalType<TIsConst, const FScope, FScope>::Type ScopeType;
+
+		// Constructor
+		TScopeTypeIterator(ScopeType* Scope)
+			: TypeIterator(Scope->TypeMap)
+		{
+			bBeforeStart = true;
+		}
+
+		/**
+		 * Indirection operator.
+		 *
+		 * @returns Type this iterator currently point to.
+		 */
+		TType* operator*()
+		{
+			return bBeforeStart ? nullptr : (TType*)(*TypeIterator).Value;
+		}
+
+		/**
+		 * Moves this iterator to next type.
+		 *
+		 * @returns True on success. False otherwise.
+		 */
+		bool MoveNext()
+		{
+			if (!bBeforeStart)
+			{
+				++TypeIterator;
+			}
+
+			bBeforeStart = false;
+			return TypeIterator.operator bool();
+		}
+
+	private:
+
+		// Current type.
+		MapIteratorType TypeIterator;
+
+		// Tells if this iterator is in initial state.
+		bool bBeforeStart;
+	};
+
+	/**
+	 * Gets type iterator.
+	 *
+	 * @return Type iterator.
+	 */
+	template <class TType = UField>
+	TScopeTypeIterator<TType, false> GetTypeIterator()
+	{
+		return TScopeTypeIterator<TType, false>(this);
+	}
+
+	/**
+	 * Gets const type iterator.
+	 *
+	 * @return Const type iterator.
+	 */
+	template <class TType = UField>
+	TScopeTypeIterator<TType, true> GetTypeIterator() const
+	{
+		return TScopeTypeIterator<TType, true>(this);
+	}
+
+	/**
+	 * Tells if this scope is a file scope.
+	 *
+	 * @returns True if this is a file scope. False otherwise.
+	 */
+	bool IsFileScope() const;
+
+	/**
+	 * Tells if this scope contains any type.
+	 *
+	 * @returns True if this scope contains type. False otherwise.
+	 */
+	bool ContainsTypes() const;
+
+	FFileScope* GetFileScope();
+private:
+	// This scopes parent.
+	const FScope* Parent;
+
+	// Map of types in this scope.
+	TMap<FName, UField*> TypeMap;
+
+	// Global map type <-> scope.
+	static TMap<UStruct*, TSharedRef<FScope> > ScopeMap;
+
+	friend struct FScopeArchiveProxy;
+	friend struct FStructScopeArchiveProxy;
+};
+
+using FScopeSet = TSet<FScope*, DefaultKeyFuncs<FScope*>, TInlineSetAllocator<1024>>;
+
+/**
+ * Represents a scope associated with source file.
+ */
+class FFileScope : public FScope
+{
+public:
+	FFileScope()
+		: SourceFile(nullptr)
+		, Name(NAME_None)
+	{ }
+	// Constructor.
+	FFileScope(FName Name, FUnrealSourceFile* SourceFile);
+
+	/**
+	 * Includes other file scopes into this scope (#include).
+	 *
+	 * @param IncludedScope File scope to include.
+	 */
+	void IncludeScope(FFileScope* IncludedScope);
+
+	virtual FFileScope* AsFileScope() override { return this; }
+
+	/**
+	 * Gets scope name.
+	 */
+	FName GetName() const override;
+
+	/**
+	 * Gets source file associated with this scope.
+	 */
+	FUnrealSourceFile* GetSourceFile() const;
+
+	/**
+	 * Appends included file scopes into given array.
+	 *
+	 * @param Out (Output parameter) Array to append scopes.
+	 */
+	void AppendIncludedFileScopes(FScopeSet& Out)
+	{
+		bool bAlreadyAdded = false;
+		Out.Add(this, &bAlreadyAdded);
+
+		if (!bAlreadyAdded)
+		{
+			for (FFileScope* IncludedScope : IncludedScopes)
+			{
+				IncludedScope->AppendIncludedFileScopes(Out);
+			}
+		}
+	}
+
+
+	const TArray<FFileScope*>& GetIncludedScopes() const
+	{
+		return IncludedScopes;
+	}
+
+	void SetSourceFile(FUnrealSourceFile* InSourceFile)
+	{
+		SourceFile = InSourceFile;
+	}
+private:
+	// Source file.
+	FUnrealSourceFile* SourceFile;
+
+	// Scope name.
+	FName Name;
+
+	// Included scopes list.
+	TArray<FFileScope*> IncludedScopes;
+
+	friend struct FFileScopeArchiveProxy;
+};
+
+
+/**
+ * Data structure representing scope of a struct or class.
+ */
+class FStructScope : public FScope
+{
+public:
+	// Constructor.
+	FStructScope(UStruct* Struct, FScope* Parent);
+
+	virtual FStructScope* AsStructScope() override { return this; }
+	/**
+	 * Gets struct associated with this scope.
+	 */
+	UStruct* GetStruct() const;
+
+	/**
+	 * Gets scope name.
+	 */
+	FName GetName() const override;
+
+private:
+	// Struct associated with this scope.
+	UStruct* Struct;
+
+	friend struct FStructScopeArchiveProxy;
+};
+
+/**
+ * Deep scope type iterator. It looks for type in the whole scope hierarchy.
+ * First going up inheritance from inner structs to outer. Then through all
+ * included scopes.
+ */
+template <class TType, bool TIsConst>
+class TDeepScopeTypeIterator
+{
+public:
+	// Conditional typedefs.
+	typedef typename TConditionalType<TIsConst, TMap<FName, UField*>::TConstIterator, TMap<FName, UField*>::TIterator>::Type MapIteratorType;
+	typedef typename TConditionalType<TIsConst, const FScope, FScope>::Type ScopeType;
+	typedef typename TConditionalType<TIsConst, const FFileScope, FFileScope>::Type FileScopeType;
+	typedef typename TConditionalType<TIsConst, FScopeSet::TConstIterator, FScopeSet::TIterator>::Type ScopeArrayIteratorType;
+
+	// Constructor.
+	TDeepScopeTypeIterator(ScopeType* Scope)
+	{
+		const ScopeType* CurrentScope = Scope;
+
+		while (!CurrentScope->IsFileScope())
+		{
+			ScopesToTraverse.Add(Scope);
+
+			UStruct* Struct = ((FStructScope*)CurrentScope)->GetStruct();
+
+			if (Struct->IsA<UClass>())
+			{
+				UClass* Class = ((UClass*)Struct)->GetSuperClass();
+
+				while (Class && !(Class->ClassFlags & EClassFlags::CLASS_Intrinsic))
+				{
+					ScopesToTraverse.Add(&FScope::GetTypeScope(Class).Get());
+					Class = Class->GetSuperClass();
+				}
+			}
+
+			CurrentScope = CurrentScope->GetParent();
+		}
+
+		((FileScopeType*)CurrentScope)->AppendIncludedFileScopes(ScopesToTraverse);
+	}
+
+	/**
+	 * Iterator increment.
+	 *
+	 * @returns True if moved iterator to another position. False otherwise.
+	 */
+	bool MoveNext()
+	{
+		if (!ScopeIterator.IsSet() && !MoveToNextScope())
+		{
+			return false;
+		}
+
+		if (ScopeIterator->MoveNext())
+		{
+			return true;
+		}
+		else
+		{
+			do
+			{
+				ScopeIterator.Reset();
+				if (!MoveToNextScope())
+				{
+					return false;
+				}
+			}
+			while(!ScopeIterator->MoveNext());
+
+			return true;
+		}
+	}
+
+	/**
+	 * Current type getter.
+	 *
+	 * @returns Current type.
+	 */
+	TType* operator*()
+	{
+		return ScopeIterator->operator*();
+	}
+
+private:
+	/**
+	 * Moves the iterator to the next scope.
+	 *
+	 * @returns True if succeeded. False otherwise.
+	 */
+	bool MoveToNextScope()
+	{
+		if (!ScopesIterator.IsSet())
+		{
+			ScopesIterator.Emplace(ScopesToTraverse);
+		}
+		else
+		{
+			ScopesIterator->operator++();
+		}
+
+		if (!ScopesIterator->operator bool())
+		{
+			return false;
+		}
+
+		ScopeIterator.Emplace(ScopesIterator->operator*());
+		return true;
+	}
+
+	// Current scope iterator.
+	TOptional<FScope::TScopeTypeIterator<TType, TIsConst> > ScopeIterator;
+
+	// Scopes list iterator.
+	TOptional<ScopeArrayIteratorType> ScopesIterator;
+
+	// List of scopes to traverse.
+	FScopeSet  ScopesToTraverse;
+};
+
diff --git a/Engine/Source/Programs/UHTLite/Private/SimplifiedParsingClassInfo.h b/Engine/Source/Programs/UHTLite/Private/SimplifiedParsingClassInfo.h
new file mode 100644
--- /dev/null
+++ b/Engine/Source/Programs/UHTLite/Private/SimplifiedParsingClassInfo.h
@@ -0,0 +1,80 @@
+// Copyright Epic Games, Inc. All Rights Reserved.
+
+#pragma once
+
+#include "CoreMinimal.h"
+#include "Containers/UnrealString.h"
+
+class FArchive;
+
+/**
+ * Simplified parsing information about UClasses.
+ */
+class FSimplifiedParsingClassInfo
+{
+public:
+	// Constructor.
+	FSimplifiedParsingClassInfo(FString&& InClassName, FString&& InBaseClassName, int32 InClassDefLine, bool bInClassIsAnInterface)
+		: ClassName          (MoveTemp(InClassName))
+		, BaseClassName      (MoveTemp(InBaseClassName))
+		, ClassDefLine       (InClassDefLine)
+		, bClassIsAnInterface(bInClassIsAnInterface)
+	{}
+
+	FSimplifiedParsingClassInfo() { }
+
+	/**
+	 * Gets class name.
+	 */
+	const FString& GetClassName() const
+	{
+		return ClassName;
+	}
+
+	/**
+	 * Gets base class name.
+	 */
+	const FString& GetBaseClassName() const
+	{
+		return BaseClassName;
+	}
+
+	/**
+	 * Gets class definition line number.
+	 */
+	int32 GetClassDefLine() const
+	{
+		return ClassDefLine;
+	}
+
+	/**
+	 * Tells if this class is an interface.
+	 */
+	bool IsInterface() const
+	{
+		return bClassIsAnInterface;
+	}
+
+	friend FArchive& operator<<(FArchive& Ar, FSimplifiedParsingClassInfo& SimplifiedParsingClassInfo)
+	{
+		Ar << SimplifiedParsingClassInfo.ClassName;
+		Ar << SimplifiedParsingClassInfo.BaseClassName;
+		Ar << SimplifiedParsingClassInfo.ClassDefLine;
+		Ar << SimplifiedParsingClassInfo.bClassIsAnInterface;
+
+		return Ar;
+	}
+
+private:
+	// Name.
+	FString ClassName;
+
+	// Super-class name.
+	FString BaseClassName;
+
+	// Class definition line.
+	int32 ClassDefLine;
+
+	// Is this an interface class?
+	bool bClassIsAnInterface;
+};
diff --git a/Engine/Source/Programs/UHTLite/Private/Specifiers/FunctionSpecifiers.h b/Engine/Source/Programs/UHTLite/Private/Specifiers/FunctionSpecifiers.h
new file mode 100644
--- /dev/null
+++ b/Engine/Source/Programs/UHTLite/Private/Specifiers/FunctionSpecifiers.h
@@ -0,0 +1,18 @@
+// Copyright Epic Games, Inc. All Rights Reserved.
+
+#pragma once
+
+#include "CoreMinimal.h"
+
+enum class EFunctionSpecifier
+{
+	None = -1
+
+	#define FUNCTION_SPECIFIER(SpecifierName) , SpecifierName
+		#include "FunctionSpecifiers.def"
+	#undef FUNCTION_SPECIFIER
+
+	, Max
+};
+
+extern const TCHAR* GFunctionSpecifierStrings[(int32)EFunctionSpecifier::Max];
diff --git a/Engine/Source/Programs/UHTLite/Private/Specifiers/InterfaceSpecifiers.h b/Engine/Source/Programs/UHTLite/Private/Specifiers/InterfaceSpecifiers.h
new file mode 100644
--- /dev/null
+++ b/Engine/Source/Programs/UHTLite/Private/Specifiers/InterfaceSpecifiers.h
@@ -0,0 +1,18 @@
+// Copyright Epic Games, Inc. All Rights Reserved.
+
+#pragma once
+
+#include "CoreMinimal.h"
+
+enum class EInterfaceSpecifier
+{
+	None = -1
+
+	#define INTERFACE_SPECIFIER(SpecifierName) , SpecifierName
+		#include "InterfaceSpecifiers.def"
+	#undef INTERFACE_SPECIFIER
+
+	, Max
+};
+
+extern const TCHAR* GInterfaceSpecifierStrings[(int32)EInterfaceSpecifier::Max];
diff --git a/Engine/Source/Programs/UHTLite/Private/Specifiers/StructSpecifiers.h b/Engine/Source/Programs/UHTLite/Private/Specifiers/StructSpecifiers.h
new file mode 100644
--- /dev/null
+++ b/Engine/Source/Programs/UHTLite/Private/Specifiers/StructSpecifiers.h
@@ -0,0 +1,18 @@
+// Copyright Epic Games, Inc. All Rights Reserved.
+
+#pragma once
+
+#include "CoreMinimal.h"
+
+enum class EStructSpecifier
+{
+	None = -1
+
+	#define STRUCT_SPECIFIER(SpecifierName) , SpecifierName
+		#include "StructSpecifiers.def"
+	#undef STRUCT_SPECIFIER
+
+	, Max
+};
+
+extern const TCHAR* GStructSpecifierStrings[(int32)EStructSpecifier::Max];
diff --git a/Engine/Source/Programs/UHTLite/Private/Specifiers/VariableSpecifiers.h b/Engine/Source/Programs/UHTLite/Private/Specifiers/VariableSpecifiers.h
new file mode 100644
--- /dev/null
+++ b/Engine/Source/Programs/UHTLite/Private/Specifiers/VariableSpecifiers.h
@@ -0,0 +1,18 @@
+// Copyright Epic Games, Inc. All Rights Reserved.
+
+#pragma once
+
+#include "CoreMinimal.h"
+
+enum class EVariableSpecifier
+{
+	None = -1
+
+	#define VARIABLE_SPECIFIER(SpecifierName) , SpecifierName
+		#include "VariableSpecifiers.def"
+	#undef VARIABLE_SPECIFIER
+
+	, Max
+};
+
+extern const TCHAR* GVariableSpecifierStrings[(int32)EVariableSpecifier::Max];
diff --git a/Engine/Source/Programs/UHTLite/Private/StringUtils.h b/Engine/Source/Programs/UHTLite/Private/StringUtils.h
new file mode 100644
--- /dev/null
+++ b/Engine/Source/Programs/UHTLite/Private/StringUtils.h
@@ -0,0 +1,43 @@
+// Copyright Epic Games, Inc. All Rights Reserved.
+
+#pragma once
+
+#include "CoreMinimal.h"
+#include "Misc/Timespan.h"
+
+/** 
+ * Attempts to strip the given class name of its affixed prefix. If no prefix exists, will return a blank string.
+ *
+ * @param InClassName Class Name with a prefix to be stripped
+ */
+FString GetClassNameWithPrefixRemoved(const FString& InClassName);
+
+/** 
+ * Attempts to strip the given class name of its affixed prefix. If no prefix exists, will return unchanged string.
+ *
+ * @param InClassName Class Name with a prefix to be stripped
+ */
+FString GetClassNameWithoutPrefix(FString InClassNameOrFilename);
+
+/** 
+ * Attempts to get class prefix. If the given class name does not start with a valid Unreal Prefix, it will return an empty string.
+ *
+ * @param InClassName Name w/ potential prefix to check
+ */
+FString GetClassPrefix(const FString& InClassName);
+
+/** 
+ * Attempts to get class prefix. If the given class name does not start with a valid Unreal Prefix, it will return an empty string.
+ *
+ * @param InClassName Name w/ potential prefix to check
+ * @param out bIsLabeledDeprecated Reference param set to True if the class name is marked as deprecated
+ */
+FString GetClassPrefix(const FString& InClassName, bool& bIsLabeledDeprecated);
+
+/**
+* Generates hash for text, ignoring characters such as \r and the terminating \0
+*
+* @param InText Text to digest
+* @return Hash generated for the specified text buffer
+*/
+uint32 GenerateTextHash(const TCHAR* InText);
diff --git a/Engine/Source/Programs/UHTLite/Private/UnderlyingEnumType.h b/Engine/Source/Programs/UHTLite/Private/UnderlyingEnumType.h
new file mode 100644
--- /dev/null
+++ b/Engine/Source/Programs/UHTLite/Private/UnderlyingEnumType.h
@@ -0,0 +1,16 @@
+// Copyright Epic Games, Inc. All Rights Reserved.
+
+#pragma once
+
+enum class EUnderlyingEnumType
+{
+	Unspecified,
+	uint8,
+	uint16,
+	uint32,
+	uint64,
+	int8,
+	int16,
+	int32,
+	int64
+};
diff --git a/Engine/Source/Programs/UHTLite/Private/UnrealHeaderTool.h b/Engine/Source/Programs/UHTLite/Private/UnrealHeaderTool.h
new file mode 100644
--- /dev/null
+++ b/Engine/Source/Programs/UHTLite/Private/UnrealHeaderTool.h
@@ -0,0 +1,12 @@
+// Copyright Epic Games, Inc. All Rights Reserved.
+#pragma once
+#include "CoreMinimal.h"
+
+struct FHeaderParserNames
+{
+	static const FName NAME_IsConversionRoot;
+	static const FName NAME_HideCategories;
+	static const FName NAME_ShowCategories;
+	static const FName NAME_SparseClassDataTypes;
+	static const FName NAME_AdvancedClassDisplay;
+};
\ No newline at end of file
diff --git a/Engine/Source/Programs/UHTLite/Private/UnrealHeaderToolGlobals.h b/Engine/Source/Programs/UHTLite/Private/UnrealHeaderToolGlobals.h
new file mode 100644
--- /dev/null
+++ b/Engine/Source/Programs/UHTLite/Private/UnrealHeaderToolGlobals.h
@@ -0,0 +1,16 @@
+// Copyright Epic Games, Inc. All Rights Reserved.
+
+#pragma once
+
+#include "CoreMinimal.h"
+#include "Logging/LogCategory.h"
+#include "Logging/LogMacros.h"
+#include "UnrealHeaderToolStats.h"
+
+DECLARE_LOG_CATEGORY_EXTERN(LogCompile, Log, All);
+
+extern bool GUHTWarningLogged;
+extern bool GUHTErrorLogged;
+
+#define UE_LOG_WARNING_UHT(Format, ...) { GUHTWarningLogged = true; UE_LOG(LogCompile, Warning, Format, ##__VA_ARGS__); }
+#define UE_LOG_ERROR_UHT(Format, ...) { GUHTErrorLogged = true; UE_LOG(LogCompile, Error, Format, ##__VA_ARGS__); }
diff --git a/Engine/Source/Programs/UHTLite/Private/UnrealHeaderToolMain.cpp b/Engine/Source/Programs/UHTLite/Private/UnrealHeaderToolMain.cpp
new file mode 100644
--- /dev/null
+++ b/Engine/Source/Programs/UHTLite/Private/UnrealHeaderToolMain.cpp
@@ -0,0 +1,121 @@
+// Copyright Epic Games, Inc. All Rights Reserved.
+
+
+#include "UnrealHeaderTool.h"
+#include "CoreMinimal.h"
+#include "Containers/UnrealString.h"
+#include "Logging/LogMacros.h"
+#include "CoreGlobals.h"
+#include "Misc/Parse.h"
+#include "Misc/CommandLine.h"
+#include "Misc/Paths.h"
+#include "Misc/ScopeExit.h"
+#include "Modules/ModuleManager.h"
+#include "Misc/CompilationResult.h"
+#include "UnrealHeaderToolGlobals.h"
+#include "RequiredProgramMainCPPInclude.h"
+
+IMPLEMENT_APPLICATION(UnrealHeaderTool, "UnrealHeaderTool");
+
+DEFINE_LOG_CATEGORY(LogCompile);
+
+bool GUHTWarningLogged = false;
+bool GUHTErrorLogged = false;
+
+/**
+ * Application entry point
+ *
+ * @param	ArgC	Command-line argument count
+ * @param	ArgV	Argument strings
+ */
+INT32_MAIN_INT32_ARGC_TCHAR_ARGV()
+{
+	FTaskTagScope Scope(ETaskTag::EGameThread);
+	FString CmdLine;
+
+	for (int32 Arg = 0; Arg < ArgC; Arg++)
+	{
+		FString LocalArg = ArgV[Arg];
+		if (LocalArg.Contains(TEXT(" "), ESearchCase::CaseSensitive))
+		{
+			CmdLine += TEXT("\"");
+			CmdLine += LocalArg;
+			CmdLine += TEXT("\"");
+		}
+		else
+		{
+			CmdLine += LocalArg;
+		}
+
+		if (Arg + 1 < ArgC)
+		{
+			CmdLine += TEXT(" ");
+		}
+	}
+
+	FString ShortCmdLine = FCommandLine::RemoveExeName(*CmdLine);
+	ShortCmdLine.TrimStartInline();	
+
+	// Get game name from the commandline. It will later be used to load the correct ini files.
+	FString ModuleInfoFilename;
+	if (ShortCmdLine.Len() && **ShortCmdLine != TEXT('-'))
+	{
+		const TCHAR* CmdLinePtr = *ShortCmdLine;
+
+		// Parse the game name or project filename.  UHT reads the list of plugins from there in case one of the plugins is UHT plugin.
+		FString GameName = FParse::Token(CmdLinePtr, false);
+
+		// This parameter is the absolute path to the file which contains information about the modules
+		// that UHT needs to generate code for.
+		ModuleInfoFilename = FParse::Token(CmdLinePtr, false );
+	}
+
+#if !NO_LOGGING
+	const static bool bVerbose = FParse::Param(*CmdLine,TEXT("VERBOSE"));
+	if (bVerbose)
+	{
+		UE_SET_LOG_VERBOSITY(LogCompile, Verbose);
+	}
+#endif
+
+	// Make sure the engine is properly cleaned up whenever we exit this function
+	ON_SCOPE_EXIT
+	{
+		FEngineLoop::AppPreExit();
+		FEngineLoop::AppExit();
+	};
+
+	GIsUCCMakeStandaloneHeaderGenerator = true;
+	if (GEngineLoop.PreInit(*ShortCmdLine) != 0)
+	{
+		UE_LOG(LogCompile, Error, TEXT("Failed to initialize the engine (PreInit failed)."));
+		return ECompilationResult::CrashOrAssert;
+	}
+
+	// Log full command line for UHT as UBT overrides LogInit verbosity settings
+	UE_LOG(LogCompile, Log, TEXT("UHT Command Line: %s"), *CmdLine);
+
+	if (ModuleInfoFilename.IsEmpty())
+	{
+		if (!FPlatformMisc::IsDebuggerPresent())
+		{
+			UE_LOG(LogCompile, Error, TEXT( "Missing module info filename on command line" ));
+			return ECompilationResult::OtherCompilationError;
+		}
+
+		// If we have a debugger, let's use a pre-existing manifest file to streamline debugging
+		// without the user having to faff around trying to get a UBT-generated manifest
+		ModuleInfoFilename = FPaths::ConvertRelativePathToFull(FPlatformProcess::BaseDir(), TEXT("../../Source/Programs/UnrealHeaderTool/Resources/UHTDebugging.manifest"));
+	}
+
+	extern ECompilationResult::Type UnrealHeaderTool_Main(const FString& ModuleInfoFilename);
+	ECompilationResult::Type Result = UnrealHeaderTool_Main(ModuleInfoFilename);
+
+	if (Result == ECompilationResult::Succeeded && (GUHTErrorLogged || (GUHTWarningLogged && GWarn->TreatWarningsAsErrors)))
+	{
+		Result = ECompilationResult::OtherCompilationError;
+	}
+
+	return Result;
+}
+
diff --git a/Engine/Source/Programs/UHTLite/Private/UnrealSourceFile.cpp b/Engine/Source/Programs/UHTLite/Private/UnrealSourceFile.cpp
new file mode 100644
--- /dev/null
+++ b/Engine/Source/Programs/UHTLite/Private/UnrealSourceFile.cpp
@@ -0,0 +1,164 @@
+// Copyright Epic Games, Inc. All Rights Reserved.
+
+#include "UnrealSourceFile.h"
+#include "UnrealHeaderTool.h"
+#include "Misc/PackageName.h"
+#include "HeaderParser.h"
+
+void FUnrealSourceFile::AddDefinedClass(UClass* Class, FSimplifiedParsingClassInfo&& ParsingInfo)
+{
+	DefinedClasses.Add(Class, MoveTemp(ParsingInfo));
+}
+
+const FString& FUnrealSourceFile::GetFileId() const
+{
+	if (FileId.Len() == 0)
+	{
+		FString StdFilename = Filename;
+
+		FPaths::MakeStandardFilename(StdFilename);
+
+		bool bRelativePath = FPaths::IsRelative(StdFilename);
+
+		if (!bRelativePath)
+		{
+			// If path is still absolute that means MakeStandardFilename has failed
+			// In this case make it relative to the current project. 
+			bRelativePath = FPaths::MakePathRelativeTo(StdFilename, *FPaths::GetPath(FPaths::GetProjectFilePath()));
+		}
+
+		// If the path has passed either MakeStandardFilename or MakePathRelativeTo it should be using internal path seperators
+		if (bRelativePath)
+		{
+			// Remove any preceding parent directory paths
+			while (StdFilename.RemoveFromStart(TEXT("../")));
+		}
+
+		FStringOutputDevice Out;
+
+		for (TCHAR Char : StdFilename)
+		{
+			if (FChar::IsAlnum(Char))
+			{
+				Out.AppendChar(Char);
+			}
+			else
+			{
+				Out.AppendChar('_');
+			}
+		}
+
+		FileId = MoveTemp(Out);
+	}
+
+	return FileId;
+}
+
+const FString& FUnrealSourceFile::GetStrippedFilename() const
+{
+	if (StrippedFilename.Len() == 0)
+	{
+		StrippedFilename = FPaths::GetBaseFilename(Filename);
+	}
+
+	return StrippedFilename;
+}
+
+FString FUnrealSourceFile::GetGeneratedMacroName(FClassMetaData* ClassData, const TCHAR* Suffix) const
+{
+	return GetGeneratedMacroName(ClassData->GetGeneratedBodyLine(), Suffix);
+}
+
+FString FUnrealSourceFile::GetGeneratedMacroName(int32 LineNumber, const TCHAR* Suffix) const
+{
+	if (Suffix != nullptr)
+	{
+		return FString::Printf(TEXT("%s_%d%s"), *GetFileId(), LineNumber, Suffix);
+	}
+
+	return FString::Printf(TEXT("%s_%d"), *GetFileId(), LineNumber);
+}
+
+FString FUnrealSourceFile::GetGeneratedBodyMacroName(int32 LineNumber, bool bLegacy) const
+{
+	return GetGeneratedMacroName(LineNumber, *FString::Printf(TEXT("%s%s"), TEXT("_GENERATED_BODY"), bLegacy ? TEXT("_LEGACY") : TEXT("")));
+}
+
+void FUnrealSourceFile::SetGeneratedFilename(FString&& InGeneratedFilename)
+{
+	GeneratedFilename = MoveTemp(InGeneratedFilename);
+}
+
+void FUnrealSourceFile::SetHasChanged(bool bInHasChanged)
+{
+	bHasChanged = bInHasChanged;
+}
+
+void FUnrealSourceFile::SetModuleRelativePath(FString&& InModuleRelativePath)
+{
+	ModuleRelativePath = MoveTemp(InModuleRelativePath);
+}
+
+void FUnrealSourceFile::SetIncludePath(FString&& InIncludePath)
+{
+	IncludePath = MoveTemp(InIncludePath);
+}
+
+const FString& FUnrealSourceFile::GetContent() const
+{
+	return Content;
+}
+
+EGeneratedCodeVersion FUnrealSourceFile::GetGeneratedCodeVersionForStruct(UStruct* Struct) const
+{
+	if (const EGeneratedCodeVersion* Version = GeneratedCodeVersions.Find(Struct))
+	{
+		return *Version;
+	}
+
+	return FHeaderParser::DefaultGeneratedCodeVersion;
+}
+
+void FUnrealSourceFile::MarkDependenciesResolved()
+{
+	bDependenciesResolved = true;
+}
+
+bool FUnrealSourceFile::AreDependenciesResolved() const
+{
+	return bDependenciesResolved;
+}
+
+void FUnrealSourceFile::SetScope(FFileScope* InScope)
+{
+	if (&Scope.Get() != InScope)
+	{
+		Scope = TSharedRef<FFileScope>(InScope);
+	}
+}
+
+void FUnrealSourceFile::SetScope(TSharedRef<FFileScope> InScope)
+{
+	Scope = InScope;
+}
+
+void FUnrealSourceFile::MarkAsParsed()
+{
+	bParsed = true;
+}
+
+bool FUnrealSourceFile::IsParsed() const
+{
+	return bParsed;
+}
+
+bool FUnrealSourceFile::HasChanged() const
+{
+	return bHasChanged;
+}
+
+FString FUnrealSourceFile::GetFileDefineName() const
+{
+	const FString API = FPackageName::GetShortName(Package).ToUpper();
+	return FString::Printf(TEXT("%s_%s_generated_h"), *API, *GetStrippedFilename());
+}
diff --git a/Engine/Source/Programs/UHTLite/Private/UnrealSourceFile.h b/Engine/Source/Programs/UHTLite/Private/UnrealSourceFile.h
new file mode 100644
--- /dev/null
+++ b/Engine/Source/Programs/UHTLite/Private/UnrealSourceFile.h
@@ -0,0 +1,314 @@
+// Copyright Epic Games, Inc. All Rights Reserved.
+
+#pragma once
+
+#include "CoreMinimal.h"
+#include "Misc/Paths.h"
+#include "Scope.h"
+#include "HeaderProvider.h"
+#include "SimplifiedParsingClassInfo.h"
+#include "GeneratedCodeVersion.h"
+
+class UPackage;
+class UClass;
+class UStruct;
+class FArchive;
+class FClassMetaData;
+
+/**
+ * Contains information about source file that defines various UHT aware types.
+ */
+class FUnrealSourceFile : public TSharedFromThis<FUnrealSourceFile>
+{
+public:
+	// Constructor.
+	FUnrealSourceFile(UPackage* InPackage, FString&& InFilename, FString&& InContent)
+		: Scope                (MakeShareable(new FFileScope(*(FString(TEXT("__")) + FPaths::GetBaseFilename(InFilename) + FString(TEXT("__File"))), this)))
+		, Filename             (MoveTemp(InFilename))
+		, Package              (InPackage)
+		, Content              (MoveTemp(InContent))
+		, bHasChanged          (false)
+		, bParsed              (false)
+		, bDependenciesResolved(false)
+	{
+		if (GetStrippedFilename() != "NoExportTypes")
+		{
+			Includes.Emplace(FHeaderProvider(EHeaderProviderSourceType::FileName, "NoExportTypes.h"));
+		}
+	}
+
+	/**
+	 * Gets parsing info for class defined in this source file.
+	 *
+	 * @returns Parsing info for class defined in this source file.
+	 */
+	const FSimplifiedParsingClassInfo& GetDefinedClassParsingInfo(UClass* DefinedClass) const
+	{
+		return GetDefinedClassesWithParsingInfo()[DefinedClass];
+	}
+
+	/**
+	 * Gets map with classes defined in this source file with parsing info.
+	 *
+	 * @returns Map with classes defined in this source file with parsing info.
+	 */
+	const TMap<UClass*, FSimplifiedParsingClassInfo>& GetDefinedClassesWithParsingInfo() const
+	{
+		return DefinedClasses;
+	}
+
+	/**
+	 * Gets number of types defined in this source file.
+	 */
+	int32 GetDefinedClassesCount() const
+	{
+		return DefinedClasses.Num();
+	}
+
+	/**
+	 * Gets generated header filename.
+	 */
+	const FString& GetGeneratedHeaderFilename() const
+	{
+		if (GeneratedHeaderFilename.Len() == 0)
+		{
+			GeneratedHeaderFilename = FString::Printf(TEXT("%s.generated.h"), *FPaths::GetBaseFilename(Filename));
+		}
+
+		return GeneratedHeaderFilename;
+	}
+
+	/**
+	 * Gets module relative path.
+	 */
+	const FString& GetModuleRelativePath() const
+	{
+		return ModuleRelativePath;
+	}
+
+	/**
+	 * Gets stripped filename.
+	 */
+	const FString& GetStrippedFilename() const;
+
+	/**
+	 * Gets unique file id.
+	 */
+	const FString& GetFileId() const;
+
+	/**
+	 * Gets define name of this source file.
+	 */
+	FString GetFileDefineName() const;
+
+	/**
+	 * Gets file-wise generated body macro name.
+	 *
+	 * @param LineNumber Number at which generated body macro is.
+	 * @param bLegacy Tells if method should get legacy generated body macro.
+	 */
+	FString GetGeneratedBodyMacroName(int32 LineNumber, bool bLegacy = false) const;
+
+	/**
+	 * Gets file-wise generated body macro name.
+	 *
+	 * @param LineNumber Number at which generated body macro is.
+	 * @param Suffix Suffix to add to generated body macro name.
+	 */
+	FString GetGeneratedMacroName(int32 LineNumber, const TCHAR* Suffix) const;
+
+	/**
+	 * Gets file-wise generated body macro name.
+	 *
+	 * @param ClassData Class metadata for which to get generated body macro name.
+	 * @param Suffix Suffix to add to generated body macro name.
+	 */
+	FString GetGeneratedMacroName(FClassMetaData* ClassData, const TCHAR* Suffix = nullptr) const;
+
+	/**
+	 * Adds given class to class definition list for this source file.
+	 *
+	 * @param Class Class to add to list.
+	 * @param ParsingInfo Data from simplified parsing step.
+	 */
+	void AddDefinedClass(UClass* Class, FSimplifiedParsingClassInfo&& ParsingInfo);
+
+	/**
+	 * Gets scope for this file.
+	 */
+	TSharedRef<FFileScope> GetScope() const
+	{
+		return Scope;
+	}
+
+	/**
+	 * Gets package this file is in.
+	 */
+	UPackage* GetPackage() const
+	{
+		return Package;
+	}
+
+	/**
+	 * Gets filename.
+	 */
+	const FString& GetFilename() const
+	{
+		return Filename;
+	}
+
+	/**
+	 * Gets generated filename.
+	 */
+	const FString& GetGeneratedFilename() const
+	{
+		return GeneratedFilename;
+	}
+
+	/**
+	 * Gets include path.
+	 */
+	const FString& GetIncludePath() const
+	{
+		return IncludePath;
+	}
+
+	/**
+	 * Gets content.
+	 */
+	const FString& GetContent() const;
+
+	/**
+	 * Gets includes.
+	 */
+	TArray<FHeaderProvider>& GetIncludes()
+	{
+		return Includes;
+	}
+
+	/**
+	 * Gets includes. Const version.
+	 */
+	const TArray<FHeaderProvider>& GetIncludes() const
+	{
+		return Includes;
+	}
+
+	/**
+	 * Gets generated code version for given UStruct.
+	 */
+	EGeneratedCodeVersion GetGeneratedCodeVersionForStruct(UStruct* Struct) const;
+
+	/**
+	 * Gets generated code versions.
+	 */
+	TMap<UStruct*, EGeneratedCodeVersion>& GetGeneratedCodeVersions()
+	{
+		return GeneratedCodeVersions;
+	}
+
+	/**
+	 * Gets generated code versions. Const version.
+	 */
+	const TMap<UStruct*, EGeneratedCodeVersion>& GetGeneratedCodeVersions() const
+	{
+		return GeneratedCodeVersions;
+	}
+
+	/**
+	 * Sets generated filename.
+	 */
+	void SetGeneratedFilename(FString&& GeneratedFilename);
+
+	/**
+	 * Sets has changed flag.
+	 */
+	void SetHasChanged(bool bHasChanged);
+
+	/**
+	 * Sets module relative path.
+	 */
+	void SetModuleRelativePath(FString&& ModuleRelativePath);
+
+	/**
+	 * Sets include path.
+	 */
+	void SetIncludePath(FString&& IncludePath);
+
+	/**
+	 * Mark this file as parsed.
+	 */
+	void MarkAsParsed();
+
+	/**
+	 * Checks if this file is parsed.
+	 */
+	bool IsParsed() const;
+
+	/**
+	 * Checks if generated file has been changed.
+	 */
+	bool HasChanged() const;
+
+	/**
+	 * Mark that this file has resolved dependencies.
+	 */
+	void MarkDependenciesResolved();
+
+	/**
+	 * Checks if dependencies has been resolved.
+	 */
+	bool AreDependenciesResolved() const;
+	friend FArchive& operator<<(FArchive& Ar, FUnrealSourceFile& UHTMakefile);
+	void SetScope(FFileScope* Scope);
+	void SetScope(TSharedRef<FFileScope> Scope);
+private:
+
+	// File scope.
+	TSharedRef<FFileScope> Scope;
+
+	// Path of this file.
+	FString Filename;
+
+	// Stripped path of this file
+	mutable FString StrippedFilename;
+
+	// Cached FileId for this file
+	mutable FString FileId;
+
+	// Package of this file.
+	UPackage* Package;
+
+	// File name of the generated header file associated with this file.
+	FString GeneratedFilename;
+
+	// File name of the generated header file associated with this file.
+	mutable FString GeneratedHeaderFilename;
+
+	// Module relative path.
+	FString ModuleRelativePath;
+
+	// Include path.
+	FString IncludePath;
+
+	// Source file content.
+	FString Content;
+
+	// Tells if generated header file was changed.
+	bool bHasChanged;
+
+	// Tells if this file was parsed.
+	bool bParsed;
+
+	// Tells if dependencies has been resolved already.
+	bool bDependenciesResolved;
+
+	// This source file includes.
+	TArray<FHeaderProvider> Includes;
+
+	// List of classes defined in this source file along with parsing info.
+	TMap<UClass*, FSimplifiedParsingClassInfo> DefinedClasses;
+
+	// Mapping of UStructs to versions, according to which their code should be generated.
+	TMap<UStruct*, EGeneratedCodeVersion> GeneratedCodeVersions;
+};
diff --git a/Engine/Source/Programs/UHTLite/Private/UnrealTypeDefinitionInfo.h b/Engine/Source/Programs/UHTLite/Private/UnrealTypeDefinitionInfo.h
new file mode 100644
--- /dev/null
+++ b/Engine/Source/Programs/UHTLite/Private/UnrealTypeDefinitionInfo.h
@@ -0,0 +1,55 @@
+// Copyright Epic Games, Inc. All Rights Reserved.
+
+#pragma once
+
+#include "CoreMinimal.h"
+#include "Templates/SharedPointer.h"
+
+// Forward declarations.
+class FUnrealSourceFile;
+
+/**
+ * Class that stores information about type (USTRUCT/UCLASS) definition.
+ */
+class FUnrealTypeDefinitionInfo : public TSharedFromThis<FUnrealTypeDefinitionInfo>
+{
+public:
+	// Constructor
+	FUnrealTypeDefinitionInfo(FUnrealSourceFile& InSourceFile, int32 InLineNumber)
+		: SourceFile(InSourceFile)
+		, LineNumber(InLineNumber)
+	{ }
+
+	/**
+	 * Gets the line number in source file this type was defined in.
+	 */
+	int32 GetLineNumber() const
+	{
+		return LineNumber;
+	}
+
+	/**
+	 * Gets the reference to FUnrealSourceFile object that stores information about
+	 * source file this type was defined in.
+	 */
+	FUnrealSourceFile& GetUnrealSourceFile()
+	{
+		return SourceFile;
+	}
+
+	void SetLineNumber(int32 InLineNumber)
+	{
+		LineNumber = InLineNumber;
+	}
+
+	const FUnrealSourceFile& GetUnrealSourceFile() const
+	{
+		return SourceFile;
+	}
+
+private:
+	FUnrealSourceFile& SourceFile;
+	int32 LineNumber;
+
+	friend struct FUnrealTypeDefinitionInfoArchiveProxy;
+};
