Branching UHT to UHTLite:

diff --git a/Engine/Source/Programs/UHTLite/Public/IScriptGeneratorPluginInterface.h b/Engine/Source/Programs/UHTLite/Public/IScriptGeneratorPluginInterface.h
new file mode 100644
--- /dev/null
+++ b/Engine/Source/Programs/UHTLite/Public/IScriptGeneratorPluginInterface.h
@@ -0,0 +1,120 @@
+// Copyright Epic Games, Inc. All Rights Reserved.
+#pragma once
+
+#include "CoreMinimal.h"
+#include "Serialization/Archive.h"
+#include "Features/IModularFeature.h"
+#include "Modules/ModuleInterface.h"
+
+class FArchive;
+class UClass;
+
+/** Build module type, mirrored in ExternalExecution.cs, enum UHTModuleType */
+struct EBuildModuleType
+{
+	enum Type
+	{
+		Program,
+		EngineRuntime,
+		EngineUncooked,
+		EngineDeveloper,
+		EngineEditor,
+		EngineThirdParty,
+		GameRuntime,
+		GameUncooked,
+		GameDeveloper,
+		GameEditor,
+		GameThirdParty,
+		// NOTE: If you add a new value, make sure to update the ToString() method below!
+
+		Max
+	};
+
+	friend FArchive& operator<<(FArchive& Ar, EBuildModuleType::Type& Type)
+	{
+		if (Ar.IsLoading())
+		{
+			uint8 Value;
+			Ar << Value;
+			Type = (EBuildModuleType::Type)Value;
+		}
+		else if (Ar.IsSaving())
+		{
+			uint8 Value = (uint8)Type;
+			Ar << Value;
+		}
+		return Ar;
+	}
+	/**
+	* Converts a string literal into EModuleType::Type value
+	*
+	* @param	The string to convert to EModuleType::Type
+	*
+	* @return	The enum value corresponding to the name
+	*/
+	static EBuildModuleType::Type Parse(const TCHAR* Value);
+};
+
+/** Build module override type to add additional PKG flags if necessary, mirrored in ModuleRules.cs, enum PackageOverrideType */
+struct EPackageOverrideType
+{
+	enum Type
+	{
+		None,
+		EditorOnly,
+		EngineDeveloper,
+		GameDeveloper,
+		EngineUncookedOnly,
+		GameUncookedOnly
+	};
+
+	friend FArchive& operator<<(FArchive& Ar, EPackageOverrideType::Type& Type)
+	{
+		if (Ar.IsLoading())
+		{
+			uint8 Value;
+			Ar << Value;
+			Type = (EPackageOverrideType::Type)Value;
+		}
+		else if (Ar.IsSaving())
+		{
+			uint8 Value = (uint8)Type;
+			Ar << Value;
+		}
+		return Ar;
+	}
+	/**
+	* Converts a string literal into EPackageOverrideType::Type value
+	*
+	* @param	The string to convert to EPackageOverrideType::Type
+	*
+	* @return	The enum value corresponding to the name
+	*/
+	static EPackageOverrideType::Type Parse(const TCHAR* Value);
+};
+
+/**
+ * The public interface to script generator plugins.
+ */
+class IScriptGeneratorPluginInterface : public IModuleInterface, public IModularFeature
+{
+public:
+
+	/** Name of module that is going to be compiling generated script glue */
+	virtual FString GetGeneratedCodeModuleName() const = 0;
+	/** Returns true if this plugin supports exporting scripts for the specified target. This should handle game as well as editor target names */
+	virtual bool SupportsTarget(const FString& TargetName) const = 0;
+	/** Returns true if this plugin supports exporting scripts for the specified module */
+	virtual bool ShouldExportClassesForModule(const FString& ModuleName, EBuildModuleType::Type ModuleType, const FString& ModuleGeneratedIncludeDirectory) const = 0;
+	/** Initializes this plugin with build information */
+	virtual void Initialize(const FString& RootLocalPath, const FString& RootBuildPath, const FString& OutputDirectory, const FString& IncludeBase) = 0;
+	/** Exports a single class. May be called multiple times for the same class (as UHT processes the entire hierarchy inside modules. */
+	virtual void ExportClass(class UClass* Class, const FString& SourceHeaderFilename, const FString& GeneratedHeaderFilename, bool bHasChanged) = 0;
+	/** Called once all classes have been exported */
+	virtual void FinishExport() = 0;
+	/** Name of the generator plugin, mostly for debuggind purposes */
+	virtual FString GetGeneratorName() const = 0;
+	/** Finds a list of external dependencies which require UHT to be re-run */
+	virtual void GetExternalDependencies(TArray<FString>& Dependencies) const { }
+};
+
diff --git a/Engine/Source/Programs/UHTLite/UHTLite.Build.cs b/Engine/Source/Programs/UHTLite/UHTLite.Build.cs
new file mode 100644
--- /dev/null
+++ b/Engine/Source/Programs/UHTLite/UHTLite.Build.cs
@@ -0,0 +1,33 @@
+// Copyright Epic Games, Inc. All Rights Reserved.
+
+using UnrealBuildTool;
+
+public class UHTLite : ModuleRules
+{
+	public UHTLite(ReadOnlyTargetRules Target) : base(Target)
+	{
+		PublicIncludePaths.Add("Runtime/Launch/Public");
+
+		PrivateDependencyModuleNames.AddRange(
+			new string[]
+			{
+				"Core",
+				"CoreUObject",
+				"Json",
+				"Projects"
+			}
+		);
+	
+		PrivateIncludePaths.AddRange(
+			new string[]
+			{
+				// For LaunchEngineLoop.cpp include
+				"Runtime/Launch/Private",
+				"Programs/UHTLite/Private",
+			});
+		
+		bEnableExceptions = true;
+
+		UnsafeTypeCastWarningLevel = WarningLevel.Warning;
+	}
+}
diff --git a/Engine/Source/Programs/UHTLite/UHTLite.Target.cs b/Engine/Source/Programs/UHTLite/UHTLite.Target.cs
new file mode 100644
--- /dev/null
+++ b/Engine/Source/Programs/UHTLite/UHTLite.Target.cs
@@ -0,0 +1,46 @@
+// Copyright Epic Games, Inc. All Rights Reserved.
+
+using UnrealBuildTool;
+using System.Collections.Generic;
+
+public class UHTLiteTarget : TargetRules
+{
+	public UHTLiteTarget(TargetInfo Target) : base(Target)
+	{
+		Type = TargetType.Program;
+		LinkType = TargetLinkType.Modular;
+		LaunchModuleName = "UHTLite";
+
+		// Never use malloc profiling in Unreal Header Tool.  We set this because often UHT is compiled right before the engine
+		// automatically by Unreal Build Tool, but if bUseMallocProfiler is defined, UHT can operate incorrectly.
+		bUseMallocProfiler = false;
+
+        bCompileICU = false;
+        // Editor-only data, however, is needed
+		bBuildWithEditorOnlyData = true;
+
+		// Currently this app is not linking against the engine, so we'll compile out references from Core to the rest of the engine
+		bCompileAgainstEngine = false;
+		bCompileAgainstApplicationCore = false;
+
+		// Force exception handling across all modules.
+		bForceEnableExceptions = true;
+
+		// Plugin support
+		bCompileWithPluginSupport = true;
+		bBuildDeveloperTools = true;
+
+		// UnrealHeaderTool is a console application, not a Windows app (sets entry point to main(), instead of WinMain())
+		bIsBuildingConsoleApplication = true;
+
+        GlobalDefinitions.Add("HACK_HEADER_GENERATOR=1");
+        GlobalDefinitions.Add("FNAME_WRITE_PROTECT_PAGES=0");
+        GlobalDefinitions.Add("USE_LOCALIZED_PACKAGE_CACHE=0");
+		GlobalDefinitions.Add("STATS=0");
+
+		// We disable compilation of GC related code since it takes a very long time (up to 30 seconds on typical hardware at the time of writing)
+		GlobalDefinitions.Add("UE_WITH_GC=0");
+		GlobalDefinitions.Add("UE_WITH_SAVEPACKAGE=0");
+		GlobalDefinitions.Add("UE_WITH_CORE_REDIRECTS=0");
+	}
+}
diff --git a/Engine/Source/Programs/UnrealBuildTool/Configuration/UEBuildTarget.cs b/Engine/Source/Programs/UnrealBuildTool/Configuration/UEBuildTarget.cs
--- a/Engine/Source/Programs/UnrealBuildTool/Configuration/UEBuildTarget.cs
+++ b/Engine/Source/Programs/UnrealBuildTool/Configuration/UEBuildTarget.cs
@@ -3431,6 +3431,9 @@ namespace UnrealBuildTool
 			// this for each file when generating their dependencies.
 			GlobalCompileEnvironment.bHackHeaderGenerator = (AppName == "UnrealHeaderTool");
 
+			// UHTLite
+			GlobalCompileEnvironment.bHackHeaderGenerator |= (AppName == "UHTLite");
+
 			GlobalCompileEnvironment.bUseDebugCRT = GlobalCompileEnvironment.Configuration == CppConfiguration.Debug && Rules.bDebugBuildsActuallyUseDebugCRT;
 			GlobalCompileEnvironment.bEnableOSX109Support = Rules.bEnableOSX109Support;
 			GlobalCompileEnvironment.Definitions.Add(String.Format("IS_PROGRAM={0}", TargetType == TargetType.Program ? "1" : "0"));
diff --git a/Engine/Source/Programs/UnrealBuildTool/System/ExternalExecution.cs b/Engine/Source/Programs/UnrealBuildTool/System/ExternalExecution.cs
--- a/Engine/Source/Programs/UnrealBuildTool/System/ExternalExecution.cs
+++ b/Engine/Source/Programs/UnrealBuildTool/System/ExternalExecution.cs
@@ -1060,6 +1060,9 @@ namespace UnrealBuildTool
 				// We never want to try to execute the header tool when we're already trying to build it!
 				bool bIsBuildingUHT = TargetName.Equals("UnrealHeaderTool", StringComparison.InvariantCultureIgnoreCase);
 
+				// UHTLite
+				bIsBuildingUHT |= TargetName.Equals("UHTLite", StringComparison.InvariantCultureIgnoreCase);
+
 				string RootLocalPath = UnrealBuildTool.RootDirectory.FullName;
 
 				UnrealTargetConfiguration UHTConfig = BuildConfiguration.bForceDebugUnrealHeaderTool ? UnrealTargetConfiguration.Debug : UnrealTargetConfiguration.Development;
