Refactored the following aspects of UHTLite:

* Globals in ClassMaps.cpp have been moved into FHeaderParser, FNativeClassHeaderGenerator, and UnrealHeaderTool_Main(). This is to illustrate how data flows through the program for further refactoring in the near future.

* Static variables have been removed from FHeaderParser.

* UnrealHeaderTool_Main() and associated functions have been moved into UnrealHeaderToolMain.cpp.

* Commented out several sections of code that are not required for code generation. This code will be refactored later.

diff --git a/Engine/Source/Programs/UHTLite/Private/BaseParser.cpp b/Engine/Source/Programs/UHTLite/Private/BaseParser.cpp
--- a/Engine/Source/Programs/UHTLite/Private/BaseParser.cpp
+++ b/Engine/Source/Programs/UHTLite/Private/BaseParser.cpp
@@ -135,6 +135,8 @@ namespace
 //////////////////////////////////////////////////////////////////////////
 // FPropertySpecifier
 
+// UHTLite NOTE: Not required for code-gen. Will be refactored later.
+/*
 FString FPropertySpecifier::ConvertToString() const
 {
 	FString Result;
@@ -170,6 +172,7 @@ FString FPropertySpecifier::ConvertToString() const
 
 	return Result;
 }
+*/
 
 /////////////////////////////////////////////////////
 // FBaseParser
@@ -712,6 +715,8 @@ bool FBaseParser::GetIdentifier( FToken& Token, bool bNoConsts )
 //
 // Get a symbol token, return 1 if gotten, 0 if not.
 //
+// UHTLite NOTE: Not required for code-gen. Will be refactored later.
+/*
 bool FBaseParser::GetSymbol( FToken& Token )
 {
 	if (!GetToken(Token))
@@ -727,6 +732,7 @@ bool FBaseParser::GetSymbol( FToken& Token )
 	UngetToken(Token);
 	return false;
 }
+*/
 
 bool FBaseParser::GetConstInt(int32& Result, const TCHAR* Tag)
 {
@@ -751,6 +757,8 @@ bool FBaseParser::GetConstInt(int32& Result, const TCHAR* Tag)
 	return false;
 }
 
+// UHTLite NOTE: Not required for code-gen. Will be refactored later.
+/*
 bool FBaseParser::GetConstInt64(int64& Result, const TCHAR* Tag)
 {
 	FToken Token;
@@ -773,6 +781,7 @@ bool FBaseParser::GetConstInt64(int64& Result, const TCHAR* Tag)
 
 	return false;
 }
+*/
 
 bool FBaseParser::MatchSymbol( const TCHAR Match, ESymbolParseOption bParseTemplateClosingBracket/*=ESymbolParseOption::Normal*/ )
 {
@@ -812,6 +821,8 @@ bool FBaseParser::MatchSymbol(const TCHAR* Match, ESymbolParseOption bParseTempl
 	return false;
 }
 
+// UHTLite NOTE: Not required for code-gen. Will be refactored later.
+/*
 //
 // Get a specific identifier and return 1 if gotten, 0 if not.
 // This is used primarily for checking for required symbols during compilation.
@@ -832,6 +843,7 @@ bool FBaseParser::MatchIdentifierByName( FName Match )
 	UngetToken(Token);
 	return false;
 }
+*/
 
 bool FBaseParser::MatchIdentifier( const TCHAR* Match, ESearchCase::Type SearchCase)
 {
@@ -919,6 +931,8 @@ bool FBaseParser::PeekSymbol( const TCHAR Match )
 	return Token.TokenType==TOKEN_Symbol && Token.Identifier[0] == Match && Token.Identifier[1] == 0;
 }
 
+// UHTLite NOTE: Not required for code-gen. Will be refactored later.
+/*
 //
 // Peek ahead and see if an identifier follows in the stream.
 //
@@ -932,6 +946,7 @@ bool FBaseParser::PeekIdentifierByName( FName Match )
 	UngetToken(Token);
 	return Token.TokenType==TOKEN_Identifier && Token.GetTokenName()==Match;
 }
+*/
 
 bool FBaseParser::PeekIdentifier( const TCHAR* Match, ESearchCase::Type SearchCase)
 {
@@ -972,6 +987,8 @@ void FBaseParser::RequireSymbol(const TCHAR Match, TFunctionRef<FString()> Error
 	}
 }
 
+// UHTLite NOTE: Not required for code-gen. Will be refactored later.
+/*
 //
 // Require an integer.
 //
@@ -982,6 +999,7 @@ void FBaseParser::RequireConstInt( const TCHAR* Match, const TCHAR* Tag )
 		FError::Throwf(TEXT("Missing integer '%s' in %s"), Match, Tag );
 	}
 }
+*/
 
 void FBaseParser::RequireAnyConstInt( const TCHAR* Tag )
 {
@@ -1216,5 +1234,3 @@ void FBaseParser::InsertMetaDataPair(TMap<FName, FString>& MetaData, FName KeyNa
 	// finally we have enough to put it into our metadata
 	MetaData.Add(KeyName, MoveTemp(Value));
 }
-
-
diff --git a/Engine/Source/Programs/UHTLite/Private/BaseParser.h b/Engine/Source/Programs/UHTLite/Private/BaseParser.h
--- a/Engine/Source/Programs/UHTLite/Private/BaseParser.h
+++ b/Engine/Source/Programs/UHTLite/Private/BaseParser.h
@@ -34,7 +34,10 @@ public:
 	FString Key;
 	TArray<FString> Values;
 
+	// UHTLite NOTE: Not required for code-gen. Will be refactored later.
+	/*
 	FString ConvertToString() const;
+	*/
 };
 
 //
@@ -129,21 +132,35 @@ public:
 
 	void UngetToken( const FToken& Token );
 	bool GetIdentifier( FToken& Token, bool bNoConsts = false );
+	
+	// UHTLite NOTE: Not required for code-gen. Will be refactored later.
+	/*
 	bool GetSymbol( FToken& Token );
+	*/
 
 	/**
 	 * Get an int constant
 	 * @return true on success, otherwise false.
 	 */
 	bool GetConstInt(int32& Result, const TCHAR* Tag = NULL);
+	
+	// UHTLite NOTE: Not required for code-gen. Will be refactored later.
+	/*
 	bool GetConstInt64(int64& Result, const TCHAR* Tag = NULL);
+	*/
 
 	// Matching predefined text.
+	// UHTLite NOTE: Not required for code-gen. Will be refactored later.
+	/*
 	bool MatchIdentifierByName( FName Match );
+	*/
 	bool MatchIdentifier( const TCHAR* Match, ESearchCase::Type SearchCase);
 	bool MatchConstInt( const TCHAR* Match );
 	bool MatchAnyConstInt();
+	// UHTLite NOTE: Not required for code-gen. Will be refactored later.
+	/*
 	bool PeekIdentifierByName( FName Match );
+	*/
 	bool PeekIdentifier( const TCHAR* Match, ESearchCase::Type SearchCase);
 	bool MatchSymbol( const TCHAR Match, ESymbolParseOption bParseTemplateClosingBracket = ESymbolParseOption::Normal );
 	bool MatchSymbol(const TCHAR* Match, ESymbolParseOption bParseTemplateClosingBracket = ESymbolParseOption::Normal);
@@ -154,7 +171,10 @@ public:
 	void RequireIdentifier( const TCHAR* Match, ESearchCase::Type SearchCase, const TCHAR* Tag );
 	void RequireSymbol( const TCHAR Match, const TCHAR* Tag, ESymbolParseOption bParseTemplateClosingBracket = ESymbolParseOption::Normal );
 	void RequireSymbol(const TCHAR Match, TFunctionRef<FString()> TagGetter, ESymbolParseOption bParseTemplateClosingBracket = ESymbolParseOption::Normal);
+	// UHTLite NOTE: Not required for code-gen. Will be refactored later.
+	/*
 	void RequireConstInt( const TCHAR* Match, const TCHAR* Tag );
+	*/
 	void RequireAnyConstInt( const TCHAR* Tag );
 
 	/** Clears out the stored comment. */
diff --git a/Engine/Source/Programs/UHTLite/Private/ClassDeclarationMetaData.cpp b/Engine/Source/Programs/UHTLite/Private/ClassDeclarationMetaData.cpp
--- a/Engine/Source/Programs/UHTLite/Private/ClassDeclarationMetaData.cpp
+++ b/Engine/Source/Programs/UHTLite/Private/ClassDeclarationMetaData.cpp
@@ -390,7 +390,7 @@ void FClassDeclarationMetaData::MergeClassCategories(FClass* Class)
 	AutoExpandCategories.Append(MoveTemp(ParentAutoExpandCategories));
 }
 
-void FClassDeclarationMetaData::MergeAndValidateClassFlags(const FString& DeclaredClassName, uint32 PreviousClassFlags, FClass* Class, const FClasses& AllClasses)
+void FClassDeclarationMetaData::MergeAndValidateClassFlags(const FString& DeclaredClassName, uint32 PreviousClassFlags, FClass* Class, const FClasses& AllClasses) const
 {
 	if (WantsToBePlaceable)
 	{
@@ -457,7 +457,7 @@ void FClassDeclarationMetaData::MergeAndValidateClassFlags(const FString& Declar
 	}
 }
 
-void FClassDeclarationMetaData::SetAndValidateConfigName(FClass* Class)
+void FClassDeclarationMetaData::SetAndValidateConfigName(FClass* Class) const
 {
 	if (ConfigName.IsEmpty() == false)
 	{
@@ -489,7 +489,7 @@ void FClassDeclarationMetaData::SetAndValidateConfigName(FClass* Class)
 	}
 }
 
-void FClassDeclarationMetaData::SetAndValidateWithinClass(FClass* Class, const FClasses& AllClasses)
+void FClassDeclarationMetaData::SetAndValidateWithinClass(FClass* Class, const FClasses& AllClasses) const
 {
 	// Process all of the class specifiers
 	if (ClassWithin.IsEmpty() == false)
diff --git a/Engine/Source/Programs/UHTLite/Private/ClassDeclarationMetaData.h b/Engine/Source/Programs/UHTLite/Private/ClassDeclarationMetaData.h
--- a/Engine/Source/Programs/UHTLite/Private/ClassDeclarationMetaData.h
+++ b/Engine/Source/Programs/UHTLite/Private/ClassDeclarationMetaData.h
@@ -55,18 +55,18 @@ public:
 	* @param	Class Class to merge flags for
 	* @param  AllClasses All known classes
 	*/
-	void MergeAndValidateClassFlags(const FString& DeclaredClassName, uint32 PreviousClassFlags, FClass* Class, const FClasses& AllClasses);
+	void MergeAndValidateClassFlags(const FString& DeclaredClassName, uint32 PreviousClassFlags, FClass* Class, const FClasses& AllClasses) const;
 private:
 
 	/** Merges all 'show' categories */
 	void MergeShowCategories();
 	/** Sets and validates 'within' property */
-	void SetAndValidateWithinClass(FClass* Class, const FClasses& AllClasses);
+	void SetAndValidateWithinClass(FClass* Class, const FClasses& AllClasses) const;
 	/** Sets and validates 'ConfigName' property */
-	void SetAndValidateConfigName(FClass* Class);
+	void SetAndValidateConfigName(FClass* Class) const;
 
 	TArray<FString> ShowCategories;
 	TArray<FString> ShowFunctions;
 	TArray<FString> DontAutoCollapseCategories;
-	bool WantsToBePlaceable;
+	mutable bool WantsToBePlaceable;
 };
diff --git a/Engine/Source/Programs/UHTLite/Private/ClassMaps.cpp b/Engine/Source/Programs/UHTLite/Private/ClassMaps.cpp
--- a/Engine/Source/Programs/UHTLite/Private/ClassMaps.cpp
+++ b/Engine/Source/Programs/UHTLite/Private/ClassMaps.cpp
@@ -4,23 +4,6 @@
 #include "UnrealHeaderTool.h"
 #include "UnrealTypeDefinitionInfo.h"
 
-FUnrealSourceFiles GUnrealSourceFilesMap;
-FTypeDefinitionInfoMap GTypeDefinitionInfoMap;
-TMap<const UPackage*, TArray<UField*>> GPackageSingletons;
-FCriticalSection GPackageSingletonsCriticalSection;
-FPublicSourceFileSet GPublicSourceFileSet;
-TMap<FProperty*, FString> GArrayDimensions;
-TMap<UPackage*,  const FManifestModule*> GPackageToManifestModuleMap;
-TMap<void*, uint32> GGeneratedCodeHashes;
-FRWLock GGeneratedCodeHashesLock;
-TMap<UEnum*,  EUnderlyingEnumType> GEnumUnderlyingTypes;
-FClassDeclarations GClassDeclarations;
-TSet<FProperty*> GUnsizedProperties;
-TSet<UField*> GEditorOnlyDataTypes;
-TMap<UStruct*, TTuple<TSharedRef<FUnrealSourceFile>, int32>> GStructToSourceLine;
-TMap<UClass*, FArchiveTypeDefinePair> GClassSerializerMap;
-TSet<FProperty*> GPropertyUsesMemoryImageAllocator;
-
 void FClassDeclarations::AddIfMissing(FName Name, TUniqueFunction<TSharedRef<FClassDeclarationMetaData>()>&& DeclConstructFunc)
 {
 	FRWScopeLock Lock(ClassDeclLock, SLT_Write);
@@ -31,17 +14,17 @@ void FClassDeclarations::AddIfMissing(FName Name, TUniqueFunction<TSharedRef<FCl
 	}
 }
 
-FClassDeclarationMetaData* FClassDeclarations::Find(FName Name)
+FClassDeclarationMetaData* FClassDeclarations::Find(FName Name) const
 {
 	FRWScopeLock Lock(ClassDeclLock, SLT_ReadOnly);
-	if (TSharedRef<FClassDeclarationMetaData>* ClassDecl = ClassDeclarations.Find(Name))
+	if (const TSharedRef<FClassDeclarationMetaData>* ClassDecl = ClassDeclarations.Find(Name))
 	{
 		return &ClassDecl->Get();
 	}
 	return nullptr;
 }
 
-FClassDeclarationMetaData& FClassDeclarations::FindChecked(FName Name)
+FClassDeclarationMetaData& FClassDeclarations::FindChecked(FName Name) const
 {
 	FRWScopeLock Lock(ClassDeclLock, SLT_ReadOnly);
 	return ClassDeclarations.FindChecked(Name).Get();
diff --git a/Engine/Source/Programs/UHTLite/Private/ClassMaps.h b/Engine/Source/Programs/UHTLite/Private/ClassMaps.h
--- a/Engine/Source/Programs/UHTLite/Private/ClassMaps.h
+++ b/Engine/Source/Programs/UHTLite/Private/ClassMaps.h
@@ -41,10 +41,26 @@ struct FTypeDefinitionInfoMap
 		DefinitionsByField.Add(Field, Definition);
 		DefinitionsByName.Add(Field->GetFName(), MoveTemp(Definition));
 	}
-	bool Contains(const UField* Field) { return DefinitionsByField.Contains(Field); }
-	TSharedRef<FUnrealTypeDefinitionInfo>* Find(const UField* Field) { return DefinitionsByField.Find(Field); }
-	TSharedRef<FUnrealTypeDefinitionInfo>* FindByName(const FName Name) { return DefinitionsByName.Find(Name); }
-	TSharedRef<FUnrealTypeDefinitionInfo>& operator[](const UField* Field) { return DefinitionsByField[Field]; }
+
+	bool Contains(const UField* Field) const
+	{
+		return DefinitionsByField.Contains(Field);
+	}
+
+	const TSharedRef<FUnrealTypeDefinitionInfo>* Find(const UField* Field) const
+	{
+		return DefinitionsByField.Find(Field);
+	}
+
+	const TSharedRef<FUnrealTypeDefinitionInfo>* FindByName(const FName Name) const
+	{
+		return DefinitionsByName.Find(Name);
+	}
+
+	const TSharedRef<FUnrealTypeDefinitionInfo>& operator[](const UField* Field) const
+	{
+		return DefinitionsByField[Field];
+	}
 
 private:
 
@@ -56,13 +72,13 @@ private:
 struct FClassDeclarations
 {
 	void AddIfMissing(FName Name, TUniqueFunction<TSharedRef<FClassDeclarationMetaData>()>&& DeclConstructFunc);
-	FClassDeclarationMetaData* Find(FName Name);
-	FClassDeclarationMetaData& FindChecked(FName Name);
+	FClassDeclarationMetaData* Find(FName Name) const;
+	FClassDeclarationMetaData& FindChecked(FName Name) const;
 
 private:
 	TMap<FName, TSharedRef<FClassDeclarationMetaData>> ClassDeclarations;
 
-	FRWLock ClassDeclLock;
+	mutable FRWLock ClassDeclLock;
 };
 
 // Wrapper class around SourceFiles map so we can quickly get a list of source files for a given package
@@ -108,23 +124,6 @@ private:
 	TMap<UPackage*, TArray<FUnrealSourceFile*>> SourceFilesByPackage;
 };
 
-extern FUnrealSourceFiles GUnrealSourceFilesMap;
-extern FTypeDefinitionInfoMap GTypeDefinitionInfoMap;
-extern TMap<const UPackage*, TArray<UField*>> GPackageSingletons;
-extern FCriticalSection GPackageSingletonsCriticalSection;
-extern FPublicSourceFileSet GPublicSourceFileSet;
-extern TMap<FProperty*, FString> GArrayDimensions;
-extern TMap<UPackage*,  const FManifestModule*> GPackageToManifestModuleMap;
-extern TMap<void*, uint32> GGeneratedCodeHashes;
-extern FRWLock GGeneratedCodeHashesLock;
-extern TMap<UEnum*, EUnderlyingEnumType> GEnumUnderlyingTypes;
-extern FClassDeclarations GClassDeclarations;
-extern TSet<FProperty*> GUnsizedProperties;
-extern TSet<UField*> GEditorOnlyDataTypes;
-extern TMap<UStruct*, TTuple<TSharedRef<FUnrealSourceFile>, int32>> GStructToSourceLine;
-extern TMap<UClass*, FArchiveTypeDefinePair> GClassSerializerMap;
-extern TSet<FProperty*> GPropertyUsesMemoryImageAllocator;
-
 /** Types access specifiers. */
 enum EAccessSpecifier
 {
diff --git a/Engine/Source/Programs/UHTLite/Private/CodeGenerator.cpp b/Engine/Source/Programs/UHTLite/Private/CodeGenerator.cpp
--- a/Engine/Source/Programs/UHTLite/Private/CodeGenerator.cpp
+++ b/Engine/Source/Programs/UHTLite/Private/CodeGenerator.cpp
@@ -54,7 +54,9 @@
 #include "Algo/Sort.h"
 #include "Algo/Reverse.h"
 #include "Async/Async.h"
+#if UHT_USE_PARALLEL_FOR
 #include "Async/ParallelFor.h"
+#endif
 #include "Misc/ScopeExit.h"
 #include "UnrealTypeDefinitionInfo.h"
 
@@ -72,24 +74,18 @@ FManifest GManifest;
 
 double GMacroizeTime = 0.0;
 
-static TArray<FString> ChangeMessages;
-static bool bWriteContents = false;
-static bool bVerifyContents = false;
-
-struct FPerHeaderData
-{
-	TSharedPtr<FUnrealSourceFile> UnrealSourceFile;
-	TArray<FHeaderProvider> DependsOn;
-	TArray<FSimplifiedParsingClassInfo> ParsedClassArray;
-};
+TArray<FString> ChangeMessages;
+bool bWriteContents = false;
+bool bVerifyContents = false;
 
-static void PerformSimplifiedClassParse(UPackage* InParent, const TCHAR* FileName, const TCHAR* Buffer, FPerHeaderData& PerHeaderData);
-static void ProcessInitialClassParse(FPerHeaderData& PerHeaderData);
+// UHTLite NOTE: Globals!
+TMap<void*, uint32> FNativeClassHeaderGenerator::GeneratedCodeHashes;
+FRWLock             FNativeClassHeaderGenerator::GeneratedCodeHashesLock;
 
 FCompilerMetadataManager GScriptHelper;
 
 // Array of all the temporary header async file tasks so we can ensure they have completed before issuing our timings
-static FGraphEventArray GAsyncFileTasks;
+FGraphEventArray GAsyncFileTasks;
 
 bool HasIdentifierExactMatch(const TCHAR* StringBegin, const TCHAR* StringEnd, const FString& Find);
 
@@ -189,119 +185,119 @@ namespace
 		return Class->FirstOwnedClassRep < Class->ClassReps.Num();
 	}
 
-	static void ExportNetData(FOutputDevice& Out, UClass* Class, const TCHAR* API)
-	{
-		const TArray<FRepRecord>& ClassReps = Class->ClassReps;
+	static const FString STRING_GetLifetimeReplicatedPropsStr(TEXT("GetLifetimeReplicatedProps"));
+}
 
-		FUHTStringBuilder NetFieldBuilder;
-		NetFieldBuilder.Logf(TEXT(""
+void FNativeClassHeaderGenerator::ExportNetData(FOutputDevice& Out, UClass* Class, const TCHAR* APIArg) const
+{
+	const TArray<FRepRecord>& ClassReps = Class->ClassReps;
+
+	FUHTStringBuilder NetFieldBuilder;
+	NetFieldBuilder.Logf(TEXT(""
 		"\tenum class ENetFields_Private : uint16\r\n"
 		"\t{\r\n"
 		"\t\tNETFIELD_REP_START=(uint16)((int32)Super::ENetFields_Private::NETFIELD_REP_END + (int32)1),\r\n"));
 
-		FUHTStringBuilder ArrayDimBuilder;
+	FUHTStringBuilder ArrayDimBuilder;
 
-		bool bAnyStaticArrays = false;
-		bool bIsFirst = true;
-		for (int32 ClassRepIndex = Class->FirstOwnedClassRep; ClassRepIndex < ClassReps.Num(); ++ClassRepIndex)
-		{
-			const FRepRecord& ClassRep = ClassReps[ClassRepIndex];
-			const FString PropertyName = ClassRep.Property->GetName();
+	bool bAnyStaticArrays = false;
+	bool bIsFirst = true;
+	for (int32 ClassRepIndex = Class->FirstOwnedClassRep; ClassRepIndex < ClassReps.Num(); ++ClassRepIndex)
+	{
+		const FRepRecord& ClassRep = ClassReps[ClassRepIndex];
+		const FString PropertyName = ClassRep.Property->GetName();
 
-			if (ClassRep.Property->ArrayDim == 1)
+		if (ClassRep.Property->ArrayDim == 1)
+		{
+			if (UNLIKELY(bIsFirst))
 			{
-				if (UNLIKELY(bIsFirst))
-				{
-					NetFieldBuilder.Logf(TEXT("\t\t%s=NETFIELD_REP_START,\r\n"), *PropertyName);
-					bIsFirst = false;
-				}
-				else
-				{
-					NetFieldBuilder.Logf(TEXT("\t\t%s,\r\n"), *PropertyName);
-				}
+				NetFieldBuilder.Logf(TEXT("\t\t%s=NETFIELD_REP_START,\r\n"), *PropertyName);
+				bIsFirst = false;
 			}
 			else
 			{
-				bAnyStaticArrays = true;
-				ArrayDimBuilder.Logf(TEXT("\t\t%s=%s,\r\n"), *PropertyName, *GArrayDimensions.FindChecked(ClassReps[ClassRepIndex].Property));
-
-				if (UNLIKELY(bIsFirst))
-				{
-					NetFieldBuilder.Logf(TEXT("\t\t%s_STATIC_ARRAY=NETFIELD_REP_START,\r\n"), *PropertyName);
-					bIsFirst = false;
-				}
-				else
-				{
-					NetFieldBuilder.Logf(TEXT("\t\t%s_STATIC_ARRAY,\r\n"), *PropertyName);
-				}
-
-				NetFieldBuilder.Logf(TEXT("\t\t%s_STATIC_ARRAY_END=((uint16)%s_STATIC_ARRAY + (uint16)EArrayDims_Private::%s - (uint16)1),\r\n"), *PropertyName, *PropertyName, *PropertyName);
+				NetFieldBuilder.Logf(TEXT("\t\t%s,\r\n"), *PropertyName);
 			}
 		}
+		else
+		{
+			bAnyStaticArrays = true;
+			ArrayDimBuilder.Logf(TEXT("\t\t%s=%s,\r\n"), *PropertyName, *HeaderParser.ArrayDimensions.FindChecked(ClassReps[ClassRepIndex].Property));
 
-		const FProperty* LastProperty = ClassReps.Last().Property;
-		NetFieldBuilder.Logf(TEXT("\t\tNETFIELD_REP_END=%s%s"), *LastProperty->GetName(), LastProperty->ArrayDim > 1 ? TEXT("_STATIC_ARRAY_END") : TEXT(""));
-
-		NetFieldBuilder.Log(TEXT("\t};"));
+			if (UNLIKELY(bIsFirst))
+			{
+				NetFieldBuilder.Logf(TEXT("\t\t%s_STATIC_ARRAY=NETFIELD_REP_START,\r\n"), *PropertyName);
+				bIsFirst = false;
+			}
+			else
+			{
+				NetFieldBuilder.Logf(TEXT("\t\t%s_STATIC_ARRAY,\r\n"), *PropertyName);
+			}
 
-		if (bAnyStaticArrays)
-		{
-			Out.Logf(TEXT(""
-				"\tenum class EArrayDims_Private : uint16\r\n"
-				"\t{\r\n"
-				"%s"
-				"\t};\r\n"), *ArrayDimBuilder);
+			NetFieldBuilder.Logf(TEXT("\t\t%s_STATIC_ARRAY_END=((uint16)%s_STATIC_ARRAY + (uint16)EArrayDims_Private::%s - (uint16)1),\r\n"), *PropertyName, *PropertyName, *PropertyName);
 		}
+	}
 
+	const FProperty* LastProperty = ClassReps.Last().Property;
+	NetFieldBuilder.Logf(TEXT("\t\tNETFIELD_REP_END=%s%s"), *LastProperty->GetName(), LastProperty->ArrayDim > 1 ? TEXT("_STATIC_ARRAY_END") : TEXT(""));
+
+	NetFieldBuilder.Log(TEXT("\t};"));
+
+	if (bAnyStaticArrays)
+	{
 		Out.Logf(TEXT(""
-			"%s\r\n" // NetFields
-			"\t%s_API virtual void ValidateGeneratedRepEnums(const TArray<struct FRepRecord>& ClassReps) const override;\r\n"),
-			*NetFieldBuilder,
-			API);
+			"\tenum class EArrayDims_Private : uint16\r\n"
+			"\t{\r\n"
+			"%s"
+			"\t};\r\n"), *ArrayDimBuilder);
 	}
 
-	static const FString STRING_GetLifetimeReplicatedPropsStr(TEXT("GetLifetimeReplicatedProps"));
+	Out.Logf(TEXT(""
+		"%s\r\n" // NetFields
+		"\t%s_API virtual void ValidateGeneratedRepEnums(const TArray<struct FRepRecord>& ClassReps) const override;\r\n"),
+		*NetFieldBuilder,
+		APIArg);
+}
 
-	static void WriteReplicatedMacroData(
-		const ClassDefinitionRange& ClassRange,
-		const TCHAR* ClassCPPName,
-		const TCHAR* API,
-		FClass* Class,
-		FClass* SuperClass,
-		FOutputDevice& Writer,
-		const FUnrealSourceFile& SourceFile,
-		FNativeClassHeaderGenerator::EExportClassOutFlags& OutFlags)
-	{
-		const bool bHasGetLifetimeReplicatedProps = HasIdentifierExactMatch(ClassRange.Start, ClassRange.End, STRING_GetLifetimeReplicatedPropsStr);
+void FNativeClassHeaderGenerator::WriteReplicatedMacroData(
+	const ClassDefinitionRange& ClassRange,
+	const TCHAR* ClassCPPName,
+	const TCHAR* APIArg,
+	FClass* Class,
+	FClass* SuperClass,
+	FOutputDevice& Writer,
+	const FUnrealSourceFile& SourceFile,
+	FNativeClassHeaderGenerator::EExportClassOutFlags& OutFlags) const
+{
+	const bool bHasGetLifetimeReplicatedProps = HasIdentifierExactMatch(ClassRange.Start, ClassRange.End, STRING_GetLifetimeReplicatedPropsStr);
 
-		if (!bHasGetLifetimeReplicatedProps)
+	if (!bHasGetLifetimeReplicatedProps)
+	{
+		// Default version autogenerates declarations.
+		if (SourceFile.GetGeneratedCodeVersionForStruct(HeaderParser, Class) == EGeneratedCodeVersion::V1)
 		{
-			// Default version autogenerates declarations.
-			if (SourceFile.GetGeneratedCodeVersionForStruct(Class) == EGeneratedCodeVersion::V1)
-			{
-				Writer.Logf(TEXT("\tvoid GetLifetimeReplicatedProps(TArray<FLifetimeProperty>& OutLifetimeProps) const override;\r\n"));
-			}
-			else
-			{
-				FError::Throwf(TEXT("Class %s has Net flagged properties and should declare member function: void GetLifetimeReplicatedProps(TArray<FLifetimeProperty>& OutLifetimeProps) const override"), ClassCPPName);
-			}
+			Writer.Logf(TEXT("\tvoid GetLifetimeReplicatedProps(TArray<FLifetimeProperty>& OutLifetimeProps) const override;\r\n"));
 		}
-
-
-		ExportNetData(Writer, Class, API);
-		
-		// If this class has replicated properties and it owns the first one, that means
-		// it's the base most replicated class. In that case, go ahead and add our interface macro.
-		if (Class->ClassReps.Num() > 0 && Class->FirstOwnedClassRep == 0)
+		else
 		{
-			OutFlags |= FNativeClassHeaderGenerator::EExportClassOutFlags::NeedsPushModelHeaders;
-			Writer.Logf(TEXT(
-				"private:\r\n"
-				"\tREPLICATED_BASE_CLASS(%s%s)\r\n"
-				"public:\r\n"
-			), Class->GetPrefixCPP(), *Class->GetName());
+			FError::Throwf(TEXT("Class %s has Net flagged properties and should declare member function: void GetLifetimeReplicatedProps(TArray<FLifetimeProperty>& OutLifetimeProps) const override"), ClassCPPName);
 		}
 	}
+
+
+	ExportNetData(Writer, Class, APIArg);
+
+	// If this class has replicated properties and it owns the first one, that means
+	// it's the base most replicated class. In that case, go ahead and add our interface macro.
+	if (Class->ClassReps.Num() > 0 && Class->FirstOwnedClassRep == 0)
+	{
+		OutFlags |= FNativeClassHeaderGenerator::EExportClassOutFlags::NeedsPushModelHeaders;
+		Writer.Logf(TEXT(
+			"private:\r\n"
+			"\tREPLICATED_BASE_CLASS(%s%s)\r\n"
+			"public:\r\n"
+		), Class->GetPrefixCPP(), *Class->GetName());
+	}
 }
 
 #define BEGIN_WRAP_EDITOR_ONLY(DoWrap) DoWrap ? BeginEditorOnlyGuard : TEXT("")
@@ -413,9 +409,9 @@ bool HasIdentifierExactMatch(const FString &String, const FString& Identifier)
 	return FindIdentifierExactMatch(String, Identifier) != INDEX_NONE;
 }
 
-void ConvertToBuildIncludePath(const UPackage* Package, FString& LocalPath)
+void ConvertToBuildIncludePath(const UPackage* Package, FString& LocalPath, const TMap<UPackage*, const FManifestModule*>& PackageToManifestModuleMap)
 {
-	FPaths::MakePathRelativeTo(LocalPath, *GPackageToManifestModuleMap.FindChecked(Package)->IncludeBase);
+	FPaths::MakePathRelativeTo(LocalPath, *PackageToManifestModuleMap.FindChecked(Package)->IncludeBase);
 }
 
 /**
@@ -473,22 +469,22 @@ FString Macroize(const TCHAR* MacroName, FString&& StringToMacroize)
 	return FString::Printf(TEXT("#define %s%s\r\n%s"), MacroName, Result.Len() ? TEXT(" \\") : TEXT(""), *Result);
 }
 
-static void AddGeneratedCodeHash(void* Field, uint32 Hash)
+void FNativeClassHeaderGenerator::AddGeneratedCodeHash(void* Field, uint32 Hash) const
 {
-	FRWScopeLock Lock(GGeneratedCodeHashesLock, SLT_Write);
-	GGeneratedCodeHashes.Add(Field, Hash);
+	FRWScopeLock Lock(GeneratedCodeHashesLock, SLT_Write);
+	GeneratedCodeHashes.Add(Field, Hash);
 }
 
 /** Generates a Hash tag string for the specified field */
-static FString GetGeneratedCodeHashTag(void* Field)
+FString FNativeClassHeaderGenerator::GetGeneratedCodeHashTag(void* Field) const
 {
 	FString Tag;
 	bool bFoundHash = false;
 	uint32 Hash = 0;
 
 	{
-		FRWScopeLock Lock(GGeneratedCodeHashesLock, SLT_ReadOnly);
-		if (const uint32* FieldHash = GGeneratedCodeHashes.Find(Field))
+		FRWScopeLock Lock(GeneratedCodeHashesLock, SLT_ReadOnly);
+		if (const uint32* FieldHash = GeneratedCodeHashes.Find(Field))
 		{
 			bFoundHash = true;
 			Hash = *FieldHash;
@@ -578,6 +574,8 @@ bool ShouldExportUFunction(UFunction* Function)
 	return true;
 }
 
+// UHTLite NOTE: Not required for code-gen. Will be refactored later.
+/*
 FString CreateLiteralString(const FString& Str)
 {
 	FString Result;
@@ -620,6 +618,7 @@ FString CreateLiteralString(const FString& Str)
 	Result += TEXT("\")");
 	return Result;
 }
+*/
 
 FString CreateUTF8LiteralString(const FString& Str)
 {
@@ -762,7 +761,7 @@ static FString OutputMetaDataCodeForObject(FOutputDevice& OutDeclaration, FOutpu
 	return Result;
 }
 
-void FNativeClassHeaderGenerator::ExportProperties(FOutputDevice& Out, UStruct* Struct, int32 TextIndent)
+void FNativeClassHeaderGenerator::ExportProperties(FOutputDevice& Out, UStruct* Struct, int32 TextIndent) const
 {
 	FProperty*	Previous			= NULL;
 	FProperty*	PreviousNonEditorOnly = NULL;
@@ -804,7 +803,7 @@ void FNativeClassHeaderGenerator::ExportProperties(FOutputDevice& Out, UStruct*
 			{
 				FUHTStringBuilder JustPropertyDecl;
 
-				const FString* Dim = GArrayDimensions.Find(Current);
+				const FString* Dim = HeaderParser.ArrayDimensions.Find(Current);
 				Current->ExportCppDeclaration( JustPropertyDecl, EExportedDeclaration::Member, Dim ? **Dim : NULL);
 				ApplyAlternatePropertyExportText(*It, JustPropertyDecl, EExportingState::TypeEraseDelegates);
 
@@ -1161,7 +1160,7 @@ void FNativeClassHeaderGenerator::PropertyNew(FOutputDevice& DeclOut, FOutputDev
 
 	if (FIntProperty* TypedProp = CastField<FIntProperty>(Prop))
 	{
-		const TCHAR* PropTypeName = GUnsizedProperties.Contains(TypedProp) ? TEXT("FUnsizedIntPropertyParams") : TEXT("FIntPropertyParams");
+		const TCHAR* PropTypeName = HeaderParser.UnsizedProperties.Contains(TypedProp) ? TEXT("FUnsizedIntPropertyParams") : TEXT("FIntPropertyParams");
 
 		DeclOut.Logf(TEXT("%sstatic const UE4CodeGen_Private::%s %s;\r\n"), DeclSpaces, PropTypeName, *NameWithoutScope);
 
@@ -1227,7 +1226,7 @@ void FNativeClassHeaderGenerator::PropertyNew(FOutputDevice& DeclOut, FOutputDev
 
 	if (FUInt32Property* TypedProp = CastField<FUInt32Property>(Prop))
 	{
-		const TCHAR* PropTypeName = GUnsizedProperties.Contains(TypedProp) ? TEXT("FUnsizedFIntPropertyParams") : TEXT("FUInt32PropertyParams");
+		const TCHAR* PropTypeName = HeaderParser.UnsizedProperties.Contains(TypedProp) ? TEXT("FUnsizedFIntPropertyParams") : TEXT("FUInt32PropertyParams");
 
 		DeclOut.Logf(TEXT("%sstatic const UE4CodeGen_Private::%s %s;\r\n"), DeclSpaces, PropTypeName, *NameWithoutScope);
 
@@ -1568,7 +1567,7 @@ void FNativeClassHeaderGenerator::PropertyNew(FOutputDevice& DeclOut, FOutputDev
 			FPropertyObjectFlags,
 			*ArrayDim,
 			OffsetStr,
-			GPropertyUsesMemoryImageAllocator.Contains(TypedProp) ? TEXT("EArrayPropertyFlags::UsesMemoryImageAllocator") : TEXT("EArrayPropertyFlags::None"),
+			HeaderParser.PropertyUsesMemoryImageAllocator.Contains(TypedProp) ? TEXT("EArrayPropertyFlags::UsesMemoryImageAllocator") : TEXT("EArrayPropertyFlags::None"),
 			*MetaDataParams,
 			*PropTag
 		);
@@ -1590,7 +1589,7 @@ void FNativeClassHeaderGenerator::PropertyNew(FOutputDevice& DeclOut, FOutputDev
 			FPropertyObjectFlags,
 			*ArrayDim,
 			OffsetStr,
-			GPropertyUsesMemoryImageAllocator.Contains(TypedProp) ? TEXT("EMapPropertyFlags::UsesMemoryImageAllocator") : TEXT("EMapPropertyFlags::None"),
+			HeaderParser.PropertyUsesMemoryImageAllocator.Contains(TypedProp) ? TEXT("EMapPropertyFlags::UsesMemoryImageAllocator") : TEXT("EMapPropertyFlags::None"),
 			*MetaDataParams,
 			*PropTag
 		);
@@ -2074,11 +2073,11 @@ const FString& FNativeClassHeaderGenerator::GetPackageSingletonName(const UPacka
 	return (*PackageSingletonStrings)->PackageSingletonName;
 }
 
-void FNativeClassHeaderGenerator::ExportGeneratedPackageInitCode(FOutputDevice& Out, const TCHAR* InDeclarations, const UPackage* InPackage, uint32 Hash)
+void FNativeClassHeaderGenerator::ExportGeneratedPackageInitCode(FOutputDevice& Out, const TCHAR* InDeclarations, const UPackage* InPackage, uint32 Hash, TMap<const UPackage*, TArray<UField*>>& PackageSingletons)
 {
 	const FString& SingletonName = GetPackageSingletonName(InPackage, nullptr);
 
-	TArray<UField*>* SingletonsToOutput = GPackageSingletons.Find(InPackage);
+	TArray<UField*>* SingletonsToOutput = PackageSingletons.Find(InPackage);
 	if (SingletonsToOutput)
 	{
 		Algo::Sort(*SingletonsToOutput, [](UField* A, UField* B)
@@ -2455,8 +2454,8 @@ void FNativeClassHeaderGenerator::ExportNativeGeneratedInitCode(FOutputDevice& O
 		while (1)
 		{
 			{
-				FRWScopeLock Lock(GGeneratedCodeHashesLock, SLT_ReadOnly);
-				if (const uint32* Hash = GGeneratedCodeHashes.Find(SuperClass))
+				FRWScopeLock Lock(GeneratedCodeHashesLock, SLT_ReadOnly);
+				if (const uint32* Hash = GeneratedCodeHashes.Find(SuperClass))
 				{
 					BaseClassHash = *Hash;
 					break;
@@ -3045,7 +3044,7 @@ void FNativeClassHeaderGenerator::ExportClassFromSourceFileInner(
 	FString PPOMacroName;
 
 	ClassDefinitionRange ClassRange;
-	if (ClassDefinitionRange* FoundRange = ClassDefinitionRanges.Find(Class))
+	if (const ClassDefinitionRange* FoundRange = HeaderParser.ClassDefinitionRanges.Find(Class))
 	{
 		ClassRange = *FoundRange;
 		ClassRange.Validate();
@@ -3054,46 +3053,51 @@ void FNativeClassHeaderGenerator::ExportClassFromSourceFileInner(
 	FString GeneratedSerializeFunctionCPP;
 	FString GeneratedSerializeFunctionHeaderMacroName;
 
-	if (const FArchiveTypeDefinePair* ArchiveTypeDefinePair = GClassSerializerMap.Find(Class))
 	{
-		FString EnclosingDefines;
-		FUHTStringBuilder Boilerplate, BoilerPlateCPP;
-		const TCHAR* MacroNameHeader;
-		const TCHAR* MacroNameCPP;
-		GeneratedSerializeFunctionHeaderMacroName = SourceFile.GetGeneratedMacroName(ClassData, TEXT("_ARCHIVESERIALIZER"));
+		// UHTLite
+		FRWScopeLock Lock(HeaderParser.ClassSerializerMapLock, SLT_ReadOnly);
 
-		EnclosingDefines = ArchiveTypeDefinePair->EnclosingDefine;
-		if (ArchiveTypeDefinePair->ArchiveType == ESerializerArchiveType::StructuredArchiveRecord)
+		if (const FArchiveTypeDefinePair* ArchiveTypeDefinePair = HeaderParser.ClassSerializerMap.Find(Class))
 		{
-			MacroNameHeader = TEXT("DECLARE_FARCHIVE_SERIALIZER");
-			MacroNameCPP = TEXT("IMPLEMENT_FARCHIVE_SERIALIZER");
-		}
-		else
-		{
-			MacroNameHeader = TEXT("DECLARE_FSTRUCTUREDARCHIVE_SERIALIZER");
-			MacroNameCPP = TEXT("IMPLEMENT_FSTRUCTUREDARCHIVE_SERIALIZER");
-		}
+			FString EnclosingDefines;
+			FUHTStringBuilder Boilerplate, BoilerPlateCPP;
+			const TCHAR* MacroNameHeader;
+			const TCHAR* MacroNameCPP;
+			GeneratedSerializeFunctionHeaderMacroName = SourceFile.GetGeneratedMacroName(ClassData, TEXT("_ARCHIVESERIALIZER"));
 
-		// if the existing Serialize function was wrapped in a compiler define directive, we need to replicate that on the generated function
-		if (EnclosingDefines.Len())
-		{
-			OutGeneratedHeaderText.Logf(TEXT("#if %s\r\n"), *EnclosingDefines);
-			BoilerPlateCPP.Logf(TEXT("#if %s\r\n"), *EnclosingDefines);
-		}
+			EnclosingDefines = ArchiveTypeDefinePair->EnclosingDefine;
+			if (ArchiveTypeDefinePair->ArchiveType == ESerializerArchiveType::StructuredArchiveRecord)
+			{
+				MacroNameHeader = TEXT("DECLARE_FARCHIVE_SERIALIZER");
+				MacroNameCPP = TEXT("IMPLEMENT_FARCHIVE_SERIALIZER");
+			}
+			else
+			{
+				MacroNameHeader = TEXT("DECLARE_FSTRUCTUREDARCHIVE_SERIALIZER");
+				MacroNameCPP = TEXT("IMPLEMENT_FSTRUCTUREDARCHIVE_SERIALIZER");
+			}
 
-		Boilerplate.Logf(TEXT("\t%s(%s, %s_API)\r\n"), MacroNameHeader, *ClassCPPName, *APIArg);
-		OutGeneratedHeaderText.Log(Macroize(*GeneratedSerializeFunctionHeaderMacroName, *Boilerplate));
-		BoilerPlateCPP.Logf(TEXT("\t%s(%s)\r\n"), MacroNameCPP, *ClassCPPName);
+			// if the existing Serialize function was wrapped in a compiler define directive, we need to replicate that on the generated function
+			if (EnclosingDefines.Len())
+			{
+				OutGeneratedHeaderText.Logf(TEXT("#if %s\r\n"), *EnclosingDefines);
+				BoilerPlateCPP.Logf(TEXT("#if %s\r\n"), *EnclosingDefines);
+			}
 
-		if (EnclosingDefines.Len())
-		{
-			OutGeneratedHeaderText.Logf(TEXT("#else\r\n"));
-			OutGeneratedHeaderText.Log(Macroize(*GeneratedSerializeFunctionHeaderMacroName, TEXT("")));
-			OutGeneratedHeaderText.Logf(TEXT("#endif\r\n"));
-			BoilerPlateCPP.Logf(TEXT("#endif\r\n"));
-		}
+			Boilerplate.Logf(TEXT("\t%s(%s, %s_API)\r\n"), MacroNameHeader, *ClassCPPName, *APIArg);
+			OutGeneratedHeaderText.Log(Macroize(*GeneratedSerializeFunctionHeaderMacroName, *Boilerplate));
+			BoilerPlateCPP.Logf(TEXT("\t%s(%s)\r\n"), MacroNameCPP, *ClassCPPName);
 
-		GeneratedSerializeFunctionCPP = BoilerPlateCPP;
+			if (EnclosingDefines.Len())
+			{
+				OutGeneratedHeaderText.Logf(TEXT("#else\r\n"));
+				OutGeneratedHeaderText.Log(Macroize(*GeneratedSerializeFunctionHeaderMacroName, TEXT("")));
+				OutGeneratedHeaderText.Logf(TEXT("#endif\r\n"));
+				BoilerPlateCPP.Logf(TEXT("#endif\r\n"));
+			}
+
+			GeneratedSerializeFunctionCPP = BoilerPlateCPP;
+		}
 	}
 
 	{
@@ -3446,11 +3450,11 @@ void ExportEnhancedConstructorsMacro(FOutputDevice& Out, FClass* Class, FClassMe
  *
  * @returns Inclusion path.
  */
-FString GetBuildPath(FUnrealSourceFile& SourceFile)
+FString GetBuildPath(FUnrealSourceFile& SourceFile, const TMap<UPackage*, const FManifestModule*>& PackageToManifestModuleMap)
 {
 	FString Out = SourceFile.GetFilename();
 
-	ConvertToBuildIncludePath(SourceFile.GetPackage(), Out);
+	ConvertToBuildIncludePath(SourceFile.GetPackage(), Out, PackageToManifestModuleMap);
 
 	return Out;
 }
@@ -3573,36 +3577,41 @@ void FNativeClassHeaderGenerator::ExportEnum(FOutputDevice& Out, UEnum* Enum) co
 	}
 	Out.Logf( TEXT("\r\n") );
 
-	// Forward declare the StaticEnum<> specialisation for enum classes
-	if (const EUnderlyingEnumType* EnumPropType = GEnumUnderlyingTypes.Find(Enum))
 	{
-		check(Enum->GetCppForm() == UEnum::ECppForm::EnumClass);
+		// UHTLite
+		FRWScopeLock Lock(HeaderParser.EnumUnderlyingTypesLock, SLT_ReadOnly);
 
-		FString UnderlyingTypeString;
-
-		if (*EnumPropType != EUnderlyingEnumType::Unspecified)
+		// Forward declare the StaticEnum<> specialisation for enum classes
+		if (const EUnderlyingEnumType* EnumPropType = HeaderParser.EnumUnderlyingTypes.Find(Enum))
 		{
-			UnderlyingTypeString = TEXT(" : ");
+			check(Enum->GetCppForm() == UEnum::ECppForm::EnumClass);
+
+			FString UnderlyingTypeString;
 
-			switch (*EnumPropType)
+			if (*EnumPropType != EUnderlyingEnumType::Unspecified)
 			{
-			case EUnderlyingEnumType::int8:        UnderlyingTypeString += TNameOf<int8>::GetName();	break;
-			case EUnderlyingEnumType::int16:       UnderlyingTypeString += TNameOf<int16>::GetName();	break;
-			case EUnderlyingEnumType::int32:       UnderlyingTypeString += TNameOf<int32>::GetName();	break;
-			case EUnderlyingEnumType::int64:       UnderlyingTypeString += TNameOf<int64>::GetName();	break;
-			case EUnderlyingEnumType::uint8:       UnderlyingTypeString += TNameOf<uint8>::GetName();	break;
-			case EUnderlyingEnumType::uint16:      UnderlyingTypeString += TNameOf<uint16>::GetName();	break;
-			case EUnderlyingEnumType::uint32:      UnderlyingTypeString += TNameOf<uint32>::GetName();	break;
-			case EUnderlyingEnumType::uint64:      UnderlyingTypeString += TNameOf<uint64>::GetName();	break;
-			default:
-				check(false);
+				UnderlyingTypeString = TEXT(" : ");
+
+				switch (*EnumPropType)
+				{
+				case EUnderlyingEnumType::int8:        UnderlyingTypeString += TNameOf<int8>::GetName();	break;
+				case EUnderlyingEnumType::int16:       UnderlyingTypeString += TNameOf<int16>::GetName();	break;
+				case EUnderlyingEnumType::int32:       UnderlyingTypeString += TNameOf<int32>::GetName();	break;
+				case EUnderlyingEnumType::int64:       UnderlyingTypeString += TNameOf<int64>::GetName();	break;
+				case EUnderlyingEnumType::uint8:       UnderlyingTypeString += TNameOf<uint8>::GetName();	break;
+				case EUnderlyingEnumType::uint16:      UnderlyingTypeString += TNameOf<uint16>::GetName();	break;
+				case EUnderlyingEnumType::uint32:      UnderlyingTypeString += TNameOf<uint32>::GetName();	break;
+				case EUnderlyingEnumType::uint64:      UnderlyingTypeString += TNameOf<uint64>::GetName();	break;
+				default:
+					check(false);
+				}
 			}
-		}
 
-		Out.Logf( TEXT("\r\n") );
-		Out.Logf( TEXT("enum class %s%s;\r\n"), *Enum->CppType, *UnderlyingTypeString );
-		Out.Logf( TEXT("template<> %sUEnum* StaticEnum<%s>();\r\n"), *GetAPIString(), *Enum->CppType );
-		Out.Logf( TEXT("\r\n") );
+			Out.Logf(TEXT("\r\n"));
+			Out.Logf(TEXT("enum class %s%s;\r\n"), *Enum->CppType, *UnderlyingTypeString);
+			Out.Logf(TEXT("template<> %sUEnum* StaticEnum<%s>();\r\n"), *GetAPIString(), *Enum->CppType);
+			Out.Logf(TEXT("\r\n"));
+		}
 	}
 }
 
@@ -3626,7 +3635,7 @@ void FNativeClassHeaderGenerator::ExportGeneratedStructBodyMacros(FOutputDevice&
 	// for RigVM methods we need to generated a macro used for implementing the static method
 	// and prepare two prologs: one for the virtual function implementation, and one for the stub
 	// invoking the static method.
-	const FRigVMStructInfo* StructRigVMInfo = FHeaderParser::StructRigVMMap.Find(Struct);
+	const FRigVMStructInfo* StructRigVMInfo = HeaderParser.StructRigVMMap.Find(Struct);
 	if(StructRigVMInfo)
 	{
 		//RigVMStubProlog.Add(FString::Printf(TEXT("ensure(RigVMOperandMemory.Num() == %d);"), StructRigVMInfo->Members.Num()));
@@ -4051,7 +4060,7 @@ void FNativeClassHeaderGenerator::ExportGeneratedEnumInitCode(FOutputDevice& Out
 	const FString EnumNameCpp           = Enum->GetName(); //UserDefinedEnum should already have a valid cpp name.
 	const FString OverriddenEnumNameCpp = FNativeClassHeaderGenerator::GetOverriddenName(Enum);
 
-	const bool bIsEditorOnlyDataType = GEditorOnlyDataTypes.Contains(Enum);
+	const bool bIsEditorOnlyDataType = HeaderParser.EditorOnlyDataTypes.Contains(Enum);
 
 	FMacroBlockEmitter EditorOnlyData(Out, TEXT("WITH_EDITORONLY_DATA"));
 	EditorOnlyData(bIsEditorOnlyDataType);
@@ -4189,7 +4198,7 @@ void FNativeClassHeaderGenerator::ExportGeneratedEnumInitCode(FOutputDevice& Out
 	Out.Log(GeneratedEnumRegisterFunctionText);
 }
 
-void FNativeClassHeaderGenerator::ExportMirrorsForNoexportStruct(FOutputDevice& Out, UScriptStruct* Struct, int32 TextIndent)
+void FNativeClassHeaderGenerator::ExportMirrorsForNoexportStruct(FOutputDevice& Out, UScriptStruct* Struct, int32 TextIndent) const
 {
 	// Export struct.
 	const FString StructName = FNameLookupCPP::GetNameCPP(Struct);
@@ -4367,7 +4376,7 @@ void FNativeClassHeaderGenerator::ExportDelegateDefinition(FOutputDevice& Out, F
 	WriteMacro(Out, MacroName, DelegateOutput);
 }
 
-void FNativeClassHeaderGenerator::ExportEventParm(FUHTStringBuilder& Out, TSet<FString>& PropertyFwd, UFunction* Function, int32 Indent, bool bOutputConstructor, EExportingState ExportingState)
+void FNativeClassHeaderGenerator::ExportEventParm(FUHTStringBuilder& Out, TSet<FString>& PropertyFwd, UFunction* Function, int32 Indent, bool bOutputConstructor, EExportingState ExportingState) const
 {
 	if (!WillExportEventParms(Function))
 	{
@@ -4415,7 +4424,7 @@ void FNativeClassHeaderGenerator::ExportEventParm(FUHTStringBuilder& Out, TSet<F
 			PropertyText.Logf(TEXT("const "));
 		}
 
-		const FString* Dim = GArrayDimensions.Find(Prop);
+		const FString* Dim = HeaderParser.ArrayDimensions.Find(Prop);
 		Prop->ExportCppDeclaration(PropertyText, EExportedDeclaration::Local, Dim ? **Dim : NULL);
 		ApplyAlternatePropertyExportText(Prop, PropertyText, ExportingState);
 
@@ -4554,7 +4563,8 @@ FString FNativeClassHeaderGenerator::GetNullParameterValue( FProperty* Prop, boo
 	return TEXT("");
 }
 
-
+// UHTLite NOTE: Not required for code-gen. Will be refactored later.
+/*
 FString FNativeClassHeaderGenerator::GetFunctionReturnString(UFunction* Function, FReferenceGatherers& OutReferenceGatherers)
 {
 	FString Result;
@@ -4576,6 +4586,7 @@ FString FNativeClassHeaderGenerator::GetFunctionReturnString(UFunction* Function
 
 	return Result;
 }
+*/
 
 /**
  * Converts Position within File to Line and Column.
@@ -4585,6 +4596,8 @@ FString FNativeClassHeaderGenerator::GetFunctionReturnString(UFunction* Function
  * @param OutLine Result line.
  * @param OutColumn Result column.
  */
+// UHTLite NOTE: Not required for code-gen. Will be refactored later.
+/*
 void GetLineAndColumnFromPositionInFile(const FString& File, int32 Position, int32& OutLine, int32& OutColumn)
 {
 	OutLine = 1;
@@ -4604,7 +4617,10 @@ void GetLineAndColumnFromPositionInFile(const FString& File, int32 Position, int
 		}
 	}
 }
+*/
 
+// UHTLite NOTE: Not required for code-gen. Will be refactored later.
+/*
 bool FNativeClassHeaderGenerator::IsMissingVirtualSpecifier(const FString& SourceFile, int32 FunctionNamePosition)
 {
 	auto IsEndOfSearchChar = [](TCHAR C) { return (C == TEXT('}')) || (C == TEXT('{')) || (C == TEXT(';')); };
@@ -4616,12 +4632,18 @@ bool FNativeClassHeaderGenerator::IsMissingVirtualSpecifier(const FString& Sourc
 	// Then find if there is "virtual" keyword starting from position of found character to ImplementationPosition
 	return !HasIdentifierExactMatch(&SourceFile[EndOfSearchCharIndex], &SourceFile[FunctionNamePosition], TEXT("virtual"));
 }
+*/
 
+// UHTLite NOTE: Not required for code-gen. Will be refactored later.
+/*
 FString CreateClickableErrorMessage(const FString& Filename, int32 Line, int32 Column)
 {
 	return FString::Printf(TEXT("%s(%d,%d): error: "), *Filename, Line, Column);
 }
+*/
 
+// UHTLite NOTE: Not required for code-gen. Will be refactored later.
+/*
 void FNativeClassHeaderGenerator::CheckRPCFunctions(FReferenceGatherers& OutReferenceGatherers, const FFuncInfo& FunctionData, const FString& ClassName, int32 ImplementationPosition, int32 ValidatePosition, const FUnrealSourceFile& SourceFile) const
 {
 	bool bHasImplementation = ImplementationPosition != INDEX_NONE;
@@ -4708,6 +4730,7 @@ void FNativeClassHeaderGenerator::CheckRPCFunctions(FReferenceGatherers& OutRefe
 		FError::Throwf(TEXT("%sDeclared function %sis not marked as virtual."), *ErrorPosition, *FunctionDecl);
 	}
 }
+*/
 
 void FNativeClassHeaderGenerator::ExportNativeFunctionHeader(
 	FOutputDevice&                   Out,
@@ -4716,8 +4739,7 @@ void FNativeClassHeaderGenerator::ExportNativeFunctionHeader(
 	EExportFunctionType::Type        FunctionType,
 	EExportFunctionHeaderStyle::Type FunctionHeaderStyle,
 	const TCHAR*                     ExtraParam,
-	const TCHAR*                     APIString
-)
+	const TCHAR*                     APIString) const
 {
 	UFunction* Function = FunctionData.FunctionReference;
 
@@ -4834,7 +4856,7 @@ void FNativeClassHeaderGenerator::ExportNativeFunctionHeader(
 
 		FUHTStringBuilder PropertyText;
 
-		const FString* Dim = GArrayDimensions.Find(Property);
+		const FString* Dim = HeaderParser.ArrayDimensions.Find(Property);
 		Property->ExportCppDeclaration( PropertyText, EExportedDeclaration::Parameter, Dim ? **Dim : NULL );
 		ApplyAlternatePropertyExportText(Property, PropertyText, EExportingState::Normal);
 
@@ -5033,9 +5055,9 @@ void FNativeClassHeaderGenerator::ExportFunctionThunk(FUHTStringBuilder& RPCWrap
 	RPCWrappers += TEXT("\t\tP_NATIVE_BEGIN;") LINE_TERMINATOR;
 
 	ClassDefinitionRange ClassRange;
-	if (ClassDefinitionRanges.Contains(Function->GetOwnerClass()))
+	if (HeaderParser.ClassDefinitionRanges.Contains(Function->GetOwnerClass()))
 	{
-		ClassRange = ClassDefinitionRanges[Function->GetOwnerClass()];
+		ClassRange = HeaderParser.ClassDefinitionRanges[Function->GetOwnerClass()];
 		ClassRange.Validate();
 	}
 
@@ -5101,7 +5123,7 @@ void FNativeClassHeaderGenerator::ExportFunctionThunk(FUHTStringBuilder& RPCWrap
 	RPCWrappers += TEXT("\t\tP_NATIVE_END;") LINE_TERMINATOR;
 }
 
-FString FNativeClassHeaderGenerator::GetFunctionParameterString(UFunction* Function, FReferenceGatherers& OutReferenceGatherers)
+FString FNativeClassHeaderGenerator::GetFunctionParameterString(UFunction* Function, FReferenceGatherers& OutReferenceGatherers) const
 {
 	FString ParameterList;
 	FUHTStringBuilder PropertyText;
@@ -5120,7 +5142,7 @@ FString FNativeClassHeaderGenerator::GetFunctionParameterString(UFunction* Funct
 			ParameterList += TEXT(", ");
 		}
 
-		FString* Dim = GArrayDimensions.Find(Property);
+		const FString* Dim = HeaderParser.ArrayDimensions.Find(Property);
 		Property->ExportCppDeclaration(PropertyText, EExportedDeclaration::Parameter, Dim ? **Dim : nullptr, 0, true);
 		ApplyAlternatePropertyExportText(Property, PropertyText, EExportingState::Normal);
 
@@ -5149,9 +5171,9 @@ void FNativeClassHeaderGenerator::ExportNativeFunctions(FOutputDevice& OutGenera
 	const FString ClassCPPName = FNameLookupCPP::GetNameCPP(Class, Class->HasAnyClassFlags(CLASS_Interface));
 
 	ClassDefinitionRange ClassRange;
-	if (ClassDefinitionRanges.Contains(Class))
+	if (HeaderParser.ClassDefinitionRanges.Contains(Class))
 	{
-		ClassRange = ClassDefinitionRanges[Class];
+		ClassRange = HeaderParser.ClassDefinitionRanges[Class];
 		ClassRange.Validate();
 	}
 
@@ -5299,9 +5321,12 @@ void FNativeClassHeaderGenerator::ExportNativeFunctions(FOutputDevice& OutGenera
 			}
 
 			// Versions that skip function autodeclaration throw an error when a function is missing.
-			if (ClassRange.bHasGeneratedBody && (SourceFile.GetGeneratedCodeVersionForStruct(Class) > EGeneratedCodeVersion::V1))
+			if (ClassRange.bHasGeneratedBody && (SourceFile.GetGeneratedCodeVersionForStruct(HeaderParser, Class) > EGeneratedCodeVersion::V1))
 			{
+				// UHTLite NOTE: Not required for code-gen. Will be refactored later.
+				/*
 				CheckRPCFunctions(OutReferenceGatherers, FunctionData, ClassCPPName, ImplementationPosition, ValidatePosition, SourceFile);
+				*/
 			}
 		}
 
@@ -5350,7 +5375,7 @@ void FNativeClassHeaderGenerator::ExportNativeFunctions(FOutputDevice& OutGenera
 
 		// Put static checks before RPCWrappers to get proper messages from static asserts before compiler errors.
 		FString NoPureDeclsMacroName = SourceFile.GetGeneratedMacroName(ClassData, TEXT("_RPC_WRAPPERS_NO_PURE_DECLS"));
-		if (SourceFile.GetGeneratedCodeVersionForStruct(Class) > EGeneratedCodeVersion::V1)
+		if (SourceFile.GetGeneratedCodeVersionForStruct(HeaderParser, Class) > EGeneratedCodeVersion::V1)
 		{
 			WriteMacro(OutGeneratedHeaderText, NoPureDeclsMacroName, RuntimeStringBuilders.RPCWrappers);
 		}
@@ -5383,7 +5408,7 @@ void FNativeClassHeaderGenerator::ExportNativeFunctions(FOutputDevice& OutGenera
 
 		// Put static checks before RPCWrappers to get proper messages from static asserts before compiler errors.
 		FString NoPureDeclsMacroName = SourceFile.GetGeneratedMacroName(ClassData, TEXT("_EDITOR_ONLY_RPC_WRAPPERS_NO_PURE_DECLS"));
-		if (SourceFile.GetGeneratedCodeVersionForStruct(Class) > EGeneratedCodeVersion::V1)
+		if (SourceFile.GetGeneratedCodeVersionForStruct(HeaderParser, Class) > EGeneratedCodeVersion::V1)
 		{
 			WriteMacro(OutGeneratedHeaderText, NoPureDeclsMacroName, EditorStringBuilders.RPCWrappers);
 		}
@@ -5422,8 +5447,7 @@ void FNativeClassHeaderGenerator::ExportCallbackFunctions(
 	const TArray<UFunction*>& CallbackFunctions,
 	const TCHAR*              CallbackWrappersMacroName,
 	EExportCallbackType       ExportCallbackType,
-	const TCHAR*              APIString
-)
+	const TCHAR*              APIString) const
 {
 	FUHTStringBuilder RPCWrappers;
 
@@ -5558,7 +5582,7 @@ void FNativeClassHeaderGenerator::ApplyAlternatePropertyExportText(FProperty* Pr
 	}
 }
 
-void GetSourceFilesInDependencyOrderRecursive(TArray<FUnrealSourceFile*>& OutTest, const UPackage* Package, FUnrealSourceFile* SourceFile, TSet<const FUnrealSourceFile*>& VisitedSet, bool bCheckDependenciesOnly, const TSet<FUnrealSourceFile*>& Ignore)
+void FNativeClassHeaderGenerator::GetSourceFilesInDependencyOrderRecursive(TArray<FUnrealSourceFile*>& OutTest, FUnrealSourceFile* SourceFile, TSet<const FUnrealSourceFile*>& VisitedSet, bool bCheckDependenciesOnly, const TSet<FUnrealSourceFile*>& Ignore) const
 {
 	// Check if the Class has already been exported, after we've checked for circular header dependencies.
 	if (OutTest.Contains(SourceFile) || Ignore.Contains(SourceFile))
@@ -5579,9 +5603,9 @@ void GetSourceFilesInDependencyOrderRecursive(TArray<FUnrealSourceFile*>& OutTes
 	VisitedSet.Add(SourceFile);
 	for (FHeaderProvider& Include : SourceFile->GetIncludes())
 	{
-		if (FUnrealSourceFile* IncludeFile = Include.Resolve())
+		if (FUnrealSourceFile* IncludeFile = Include.Resolve(HeaderParser.UnrealSourceFilesMap, HeaderParser.TypeDefinitionInfoMap))
 		{
-			GetSourceFilesInDependencyOrderRecursive(OutTest, Package, IncludeFile, VisitedSet, bCheckDependenciesOnly, Ignore);
+			GetSourceFilesInDependencyOrderRecursive(OutTest, IncludeFile, VisitedSet, bCheckDependenciesOnly, Ignore);
 		}
 	}
 	VisitedSet.Remove(SourceFile);
@@ -5592,7 +5616,7 @@ void GetSourceFilesInDependencyOrderRecursive(TArray<FUnrealSourceFile*>& OutTes
 	}
 }
 
-TArray<FUnrealSourceFile*> GetSourceFilesInDependencyOrder(const UPackage* Package, const TSet<FUnrealSourceFile*>& SourceFiles, const TSet<FUnrealSourceFile*>& Ignore)
+TArray<FUnrealSourceFile*> FNativeClassHeaderGenerator::GetSourceFilesInDependencyOrder(const TSet<FUnrealSourceFile*>& SourceFiles, const TSet<FUnrealSourceFile*>& Ignore) const
 {
 	TArray<FUnrealSourceFile*> Result;
 	TSet<const FUnrealSourceFile*>	VisitedSet;
@@ -5600,7 +5624,7 @@ TArray<FUnrealSourceFile*> GetSourceFilesInDependencyOrder(const UPackage* Packa
 	{
 		if (SourceFile->GetPackage() == Package)
 		{
-			GetSourceFilesInDependencyOrderRecursive(Result, Package, SourceFile, VisitedSet, false, Ignore);
+			GetSourceFilesInDependencyOrderRecursive(Result, SourceFile, VisitedSet, false, Ignore);
 		}
 	}
 
@@ -5618,7 +5642,9 @@ static bool HasDynamicOuter(UField* Field)
 static void RecordPackageSingletons(
 	const UPackage& Package,
 	const TArray<UScriptStruct*>& Structs,
-	const TArray<UDelegateFunction*>& Delegates)
+	const TArray<UDelegateFunction*>& Delegates,
+	TMap<const UPackage*, TArray<UField*>>& PackageSingletons,
+	FCriticalSection& PackageSingletonsCriticalSection)
 {
 	TArray<UField*> Singletons;
 	Singletons.Reserve(Structs.Num() + Delegates.Num());
@@ -5640,10 +5666,10 @@ static void RecordPackageSingletons(
 
 	if (Singletons.Num())
 	{
-		FScopeLock PackageSingletonLock(&GPackageSingletonsCriticalSection);
+		FScopeLock PackageSingletonLock(&PackageSingletonsCriticalSection);
 
-		TArray<UField*>& PackageSingletons = GPackageSingletons.FindOrAdd(&Package);
-		PackageSingletons.Append(MoveTemp(Singletons));
+		TArray<UField*>& CurrentPackageSingletons = PackageSingletons.FindOrAdd(&Package);
+		CurrentPackageSingletons.Append(MoveTemp(Singletons));
 	}
 }
 
@@ -5726,12 +5752,16 @@ private:
 
 // Constructor.
 FNativeClassHeaderGenerator::FNativeClassHeaderGenerator(
+	const FHeaderParser& InHeaderParser,
 	const UPackage* InPackage,
 	const TSet<FUnrealSourceFile*>& SourceFiles,
 	FClasses& AllClasses,
-	bool InAllowSaveExportedHeaders
+	bool InAllowSaveExportedHeaders,
+	const FPublicSourceFileSet& PublicSourceFileSet,
+	const TMap<UPackage*, const FManifestModule*>& PackageToManifestModuleMap
 )
-	: API                        (FPackageName::GetShortName(InPackage).ToUpper())
+	: HeaderParser(InHeaderParser)
+	, API                        (FPackageName::GetShortName(InPackage).ToUpper())
 	, APIStringPrivate           (FString::Printf(TEXT("%s_API "), *API))
 	, Package                    (InPackage)
 	, bAllowSaveExportedHeaders  (InAllowSaveExportedHeaders)
@@ -5760,7 +5790,7 @@ FNativeClassHeaderGenerator::FNativeClassHeaderGenerator(
 				{
 					Class->UnMark(EObjectMark(OBJECTMARK_TagImp | OBJECTMARK_TagExp));
 				}
-				else if (!Class->HasAnyClassFlags(CLASS_NoExport) && GTypeDefinitionInfoMap.Contains(Class))
+				else if (!Class->HasAnyClassFlags(CLASS_NoExport) && HeaderParser.TypeDefinitionInfoMap.Contains(Class))
 				{
 					bWriteClassesH = true;
 					Class->UnMark(OBJECTMARK_TagImp);
@@ -5774,7 +5804,7 @@ FNativeClassHeaderGenerator::FNativeClassHeaderGenerator(
 	{
 		// Get source files and ignore them next time round
 		static TSet<FUnrealSourceFile*> ExportedSourceFiles;
-		Exported = GetSourceFilesInDependencyOrder(Package, SourceFiles, ExportedSourceFiles);
+		Exported = GetSourceFilesInDependencyOrder(SourceFiles, ExportedSourceFiles);
 		ExportedSourceFiles.Append(Exported);
 	}
 
@@ -5801,7 +5831,7 @@ FNativeClassHeaderGenerator::FNativeClassHeaderGenerator(
 	for (FUnrealSourceFile* SourceFile : Exported)
 	{
 		FString ModuleRelativeFilename = SourceFile->GetFilename();
-		ConvertToBuildIncludePath(Package, ModuleRelativeFilename);
+		ConvertToBuildIncludePath(Package, ModuleRelativeFilename, PackageToManifestModuleMap);
 
 		FString StrippedName = FPaths::GetBaseFilename(ModuleRelativeFilename);
 		FString GeneratedSourceFilename = (PackageManifest->GeneratedIncludeDirectory / StrippedName) + TEXT(".gen.cpp");
@@ -5830,22 +5860,37 @@ FNativeClassHeaderGenerator::FNativeClassHeaderGenerator(
 	TArray<FPreloadHeaderFileInfo> PreloadedFiles;
 	PreloadedFiles.SetNum(Exported.Num());
 
-	ParallelFor(Exported.Num(), [&Exported, &PreloadedFiles, Package=Package, ConstPackageManifest](int32 Index)
+#if	UHT_USE_PARALLEL_FOR
+	ParallelFor(Exported.Num(), [&Exported, &PreloadedFiles, Package=Package, ConstPackageManifest, &PackageToManifestModuleMap](int32 Index)
+#else
+	for(int Index = 0; Index < Exported.Num(); Index++)
+#endif
 	{
 		FUnrealSourceFile* SourceFile = Exported[Index];
 
 		FString ModuleRelativeFilename = SourceFile->GetFilename();
-		ConvertToBuildIncludePath(Package, ModuleRelativeFilename);
+		ConvertToBuildIncludePath(Package, ModuleRelativeFilename, PackageToManifestModuleMap);
 
 		FString StrippedName = FPaths::GetBaseFilename(MoveTemp(ModuleRelativeFilename));
 		FString HeaderPath = (ConstPackageManifest->GeneratedIncludeDirectory / StrippedName) + TEXT(".generated.h");
 
 		PreloadedFiles[Index].Load(MoveTemp(HeaderPath));
-	});
+	}
+#if	UHT_USE_PARALLEL_FOR
+	);
+#endif
+
+	// UHTLite
+	TMap<const UPackage*, TArray<UField*>> PackageSingletons;
+	FCriticalSection PackageSingletonsCriticalSection;
 
 	FString ExceptionMessage;
 
-	ParallelFor(Exported.Num(), [&Exported, &PreloadedFiles, &GeneratedCPPs, &TempSaveTasks=TempSaveTasks, &ExceptionMessage, ConstThis](int32 Index)
+#if	UHT_USE_PARALLEL_FOR
+	ParallelFor(Exported.Num(), [&Exported, &PreloadedFiles, &GeneratedCPPs, &TempSaveTasks=TempSaveTasks, &ExceptionMessage, ConstThis, &PackageSingletons, &PackageSingletonsCriticalSection](int32 Index)
+#else
+	for(int Index = 0; Index < Exported.Num(); Index++)
+#endif
 	{
 #if !PLATFORM_EXCEPTIONS_DISABLED
 		try
@@ -5869,7 +5914,7 @@ FNativeClassHeaderGenerator::FNativeClassHeaderGenerator(
 			TArray<UDelegateFunction*> DelegateFunctions;
 			SourceFile->GetScope()->SplitTypesIntoArrays(Enums, Structs, DelegateFunctions);
 
-			RecordPackageSingletons(*SourceFile->GetPackage(), Structs, DelegateFunctions);
+			RecordPackageSingletons(*SourceFile->GetPackage(), Structs, DelegateFunctions, PackageSingletons, PackageSingletonsCriticalSection);
 
 			// Reverse the containers as they come out in the reverse order of declaration
 			Algo::Reverse(Enums);
@@ -5963,7 +6008,10 @@ FNativeClassHeaderGenerator::FNativeClassHeaderGenerator(
 			}
 		}
 #endif
-	});
+	}
+#if	UHT_USE_PARALLEL_FOR
+	);
+#endif
 
 	if (ExceptionMessage.Len() > 0)
 	{
@@ -5990,7 +6038,7 @@ FNativeClassHeaderGenerator::FNativeClassHeaderGenerator(
 			GClassToSourceFileMap.Add(Class, SourceFile);
 		}
 
-		if (GPublicSourceFileSet.Contains(SourceFile))
+		if (PublicSourceFileSet.Contains(SourceFile))
 		{
 			PublicHeaderGroupIncludes.Add(SourceFile);
 		}
@@ -6013,23 +6061,28 @@ FNativeClassHeaderGenerator::FNativeClassHeaderGenerator(
 			{
 				if (FUnrealSourceFile** WithinSourceFile = GClassToSourceFileMap.Find(Class->ClassWithin))
 				{
-					FString Header = GetBuildPath(**WithinSourceFile);
+					FString Header = GetBuildPath(**WithinSourceFile, PackageToManifestModuleMap);
 					RelativeIncludes.AddUnique(MoveTemp(Header));
 				}
 			}
 
-			if (const FArchiveTypeDefinePair* ArchiveTypeDefinePair = GClassSerializerMap.Find(Class))
 			{
-				if (!bAddedStructuredArchiveFromArchiveHeader && ArchiveTypeDefinePair->ArchiveType == ESerializerArchiveType::StructuredArchiveRecord)
-				{
-					RelativeIncludes.AddUnique(TEXT("Serialization/StructuredArchive.h"));
-					bAddedStructuredArchiveFromArchiveHeader = true;
-				}
+				// UHTLite
+				FRWScopeLock Lock(HeaderParser.ClassSerializerMapLock, SLT_ReadOnly);
 
-				if (!bAddedArchiveUObjectFromStructuredArchiveHeader && ArchiveTypeDefinePair->ArchiveType == ESerializerArchiveType::Archive)
+				if (const FArchiveTypeDefinePair* ArchiveTypeDefinePair = HeaderParser.ClassSerializerMap.Find(Class))
 				{
-					RelativeIncludes.AddUnique(TEXT("Serialization/ArchiveUObjectFromStructuredArchive.h"));
-					bAddedArchiveUObjectFromStructuredArchiveHeader = true;
+					if (!bAddedStructuredArchiveFromArchiveHeader && ArchiveTypeDefinePair->ArchiveType == ESerializerArchiveType::StructuredArchiveRecord)
+					{
+						RelativeIncludes.AddUnique(TEXT("Serialization/StructuredArchive.h"));
+						bAddedStructuredArchiveFromArchiveHeader = true;
+					}
+
+					if (!bAddedArchiveUObjectFromStructuredArchiveHeader && ArchiveTypeDefinePair->ArchiveType == ESerializerArchiveType::Archive)
+					{
+						RelativeIncludes.AddUnique(TEXT("Serialization/ArchiveUObjectFromStructuredArchive.h"));
+						bAddedArchiveUObjectFromStructuredArchiveHeader = true;
+					}
 				}
 			}
 		}
@@ -6048,14 +6101,14 @@ FNativeClassHeaderGenerator::FNativeClassHeaderGenerator(
 		ClassesHText.Log(TEXT("\r\n"));
 
 		// Fill with the rest source files from this package.
-		if (const TArray<FUnrealSourceFile*>* SourceFilesForPackage = GPublicSourceFileSet.FindFilesForPackage(InPackage))
+		if (const TArray<FUnrealSourceFile*>* SourceFilesForPackage = PublicSourceFileSet.FindFilesForPackage(InPackage))
 		{
 			PublicHeaderGroupIncludes.Append(*SourceFilesForPackage);
 		}
 
 		for (FUnrealSourceFile* SourceFile : PublicHeaderGroupIncludes)
 		{
-			ClassesHText.Logf(TEXT("#include \"%s\"") LINE_TERMINATOR, *GetBuildPath(*SourceFile));
+			ClassesHText.Logf(TEXT("#include \"%s\"") LINE_TERMINATOR, *GetBuildPath(*SourceFile, PackageToManifestModuleMap));
 		}
 
 		ClassesHText.Log(LINE_TERMINATOR);
@@ -6090,10 +6143,10 @@ FNativeClassHeaderGenerator::FNativeClassHeaderGenerator(
 		}
 
 		FGeneratedCPP& GeneratedCPP = GeneratedCPPs.Emplace(nullptr, PackageManifest->GeneratedIncludeDirectory / FString::Printf(TEXT("%s.init.gen.cpp"), *PackageName));
-		ExportGeneratedPackageInitCode(GeneratedCPP.GeneratedText, *GeneratedFunctionDeclarations, Package, CombinedHash);
+		ExportGeneratedPackageInitCode(GeneratedCPP.GeneratedText, *GeneratedFunctionDeclarations, Package, CombinedHash, PackageSingletons);
 	}
 
-	const FManifestModule* ModuleInfo = GPackageToManifestModuleMap.FindChecked(Package);
+	const FManifestModule* ModuleInfo = PackageToManifestModuleMap.FindChecked(Package);
 
 	struct FGeneratedCPPInfo
 	{
@@ -6112,18 +6165,29 @@ FNativeClassHeaderGenerator::FNativeClassHeaderGenerator(
 
 	if (bAllowSaveExportedHeaders)
 	{
+#if	UHT_USE_PARALLEL_FOR
 		ParallelFor(GeneratedCPPArray.Num(), [&](int32 Index)
+#else
+		for(Index = 0; Index < GeneratedCPPArray.Num(); Index++)
+#endif
 		{
 			FGeneratedCPPInfo& CPPInfo = GeneratedCPPArray[Index];
 			CPPInfo.FileInfo.Load(FString(CPPInfo.GeneratedCPP->GeneratedCppFullFilename));
-		});
+		}
+#if	UHT_USE_PARALLEL_FOR
+		);
+#endif
 	}
 
 	const int32 SaveTaskStartIndex = TempSaveTasks.Num();
 	TempSaveTasks.AddDefaulted(GeneratedCPPArray.Num());
 
 	// Generate CPP files
+#if	UHT_USE_PARALLEL_FOR
 	ParallelFor(GeneratedCPPArray.Num(), [ConstThis, &GeneratedCPPArray, &TempSaveTasks=TempSaveTasks, SaveTaskStartIndex](int32 Index)
+#else
+	for(Index = 0; Index < GeneratedCPPArray.Num(); Index++)
+#endif
 	{
 		FGeneratedCPPInfo& CPPInfo = GeneratedCPPArray[Index];
 		FGeneratedCPP* GeneratedCPP = CPPInfo.GeneratedCPP;
@@ -6151,7 +6215,10 @@ FNativeClassHeaderGenerator::FNativeClassHeaderGenerator(
 		);
 
 		ConstThis->SaveHeaderIfChanged(ReferenceGatherers, CPPInfo.FileInfo, MoveTemp(FileText), TempSaveTasks[SaveTaskStartIndex+Index]);
-	});
+	}
+#if	UHT_USE_PARALLEL_FOR
+	);
+#endif
 
 	if (bAllowSaveExportedHeaders)
 	{
@@ -6343,7 +6410,13 @@ bool FNativeClassHeaderGenerator::SaveHeaderIfChanged(FReferenceGatherers& OutRe
 	const bool bHasChanged = OriginalHeaderLocal.Len() == 0 || FCString::Strcmp(*OriginalHeaderLocal, *InNewHeaderContents);
 	if (bHasChanged)
 	{
+		// UHTLite NOTE: Always fail on differences, because we test against UHT output to ensure consistency.
+		static const bool bFailIfGeneratedCodeChanges = true;
+		UE_DEBUG_BREAK();
+		/*
 		static const bool bFailIfGeneratedCodeChanges = FParse::Param(FCommandLine::Get(), TEXT("FailIfGeneratedCodeChanges"));
+		*/
+
 		if (bFailIfGeneratedCodeChanges)
 		{
 			FString ConflictPath = HeaderPathStr + TEXT(".conflict");
@@ -6393,7 +6466,11 @@ void FNativeClassHeaderGenerator::ExportUpdatedHeaders(FString&& PackageName, TA
 	{
 		auto MoveHeadersTask = [PackageName = MoveTemp(PackageName), TempHeaderPaths = MoveTemp(TempHeaderPaths)]()
 		{
+#if	UHT_USE_PARALLEL_FOR
 			ParallelFor(TempHeaderPaths.Num(), [&](int32 Index)
+#else
+			for(int Index = 0; Index < TempHeaderPaths.Num(); Index++)
+#endif
 			{
 				const FString& TmpFilename = TempHeaderPaths[Index];
 				FString Filename = GenerateTempHeaderName(TmpFilename, true);
@@ -6405,7 +6482,10 @@ void FNativeClassHeaderGenerator::ExportUpdatedHeaders(FString&& PackageName, TA
 				{
 					UE_LOG(LogCompile, Log, TEXT("Exported updated C++ header: %s"), *Filename);
 				}
-			});
+			}
+#if	UHT_USE_PARALLEL_FOR
+			);
+#endif
 		};
 
 		FTaskGraphInterface::Get().WaitUntilTasksComplete(InTempSaveTasks);
@@ -6442,735 +6522,3 @@ void FNativeClassHeaderGenerator::ExportGeneratedCPP(FOutputDevice& Out, const T
 	Out.Log(EnableDeprecationWarnings);
 	Out.Log(EnableWarning4883);
 }
-
-/** Get all script plugins based on ini setting */
-void GetScriptPlugins(TArray<IScriptGeneratorPluginInterface*>& ScriptPlugins)
-{
-	FScopedDurationTimer PluginTimeTracker(GPluginOverheadTime);
-
-	ScriptPlugins = IModularFeatures::Get().GetModularFeatureImplementations<IScriptGeneratorPluginInterface>(TEXT("ScriptGenerator"));
-	UE_LOG(LogCompile, Log, TEXT("Found %d script generator plugins."), ScriptPlugins.Num());
-
-	// Check if we can use these plugins and initialize them
-	for (int32 PluginIndex = ScriptPlugins.Num() - 1; PluginIndex >= 0; --PluginIndex)
-	{
-		IScriptGeneratorPluginInterface* ScriptGenerator = ScriptPlugins[PluginIndex];
-		bool bSupportedPlugin = ScriptGenerator->SupportsTarget(GManifest.TargetName);
-		if (bSupportedPlugin)
-		{
-			// Find the right output directory for this plugin base on its target (Engine-side) plugin name.
-			FString GeneratedCodeModuleName = ScriptGenerator->GetGeneratedCodeModuleName();
-			const FManifestModule* GeneratedCodeModule = NULL;
-			FString OutputDirectory;
-			FString IncludeBase;
-			for (const FManifestModule& Module : GManifest.Modules)
-			{
-				if (Module.Name == GeneratedCodeModuleName)
-				{
-					GeneratedCodeModule = &Module;
-				}
-			}
-			if (GeneratedCodeModule)
-			{
-				UE_LOG(LogCompile, Log, TEXT("Initializing script generator \'%s\'"), *ScriptGenerator->GetGeneratorName());
-				ScriptGenerator->Initialize(GManifest.RootLocalPath, GManifest.RootBuildPath, GeneratedCodeModule->GeneratedIncludeDirectory, GeneratedCodeModule->IncludeBase);
-			}
-			else
-			{
-				// Can't use this plugin
-				UE_LOG(LogCompile, Log, TEXT("Unable to determine output directory for %s. Cannot export script glue with \'%s\'"), *GeneratedCodeModuleName, *ScriptGenerator->GetGeneratorName());
-				bSupportedPlugin = false;
-			}
-		}
-		if (!bSupportedPlugin)
-		{
-			UE_LOG(LogCompile, Log, TEXT("Script generator \'%s\' not supported for target: %s"), *ScriptGenerator->GetGeneratorName(), *GManifest.TargetName);
-			ScriptPlugins.RemoveAt(PluginIndex);
-		}
-	}
-}
-
-/**
- * Tries to resolve super classes for classes defined in the given
- * module.
- *
- * @param Package Modules package.
- */
-void ResolveSuperClasses(UPackage* Package)
-{
-	TArray<UObject*> Objects;
-	GetObjectsWithPackage(Package, Objects);
-
-	for (UObject* Object : Objects)
-	{
-		if (!Object->IsA<UClass>() || Object->HasAnyFlags(RF_ClassDefaultObject))
-		{
-			continue;
-		}
-
-		UClass* DefinedClass = Cast<UClass>(Object);
-
-		if (DefinedClass->HasAnyClassFlags(CLASS_Intrinsic | CLASS_NoExport))
-		{
-			continue;
-		}
-
-		const FSimplifiedParsingClassInfo& ParsingInfo = GTypeDefinitionInfoMap[DefinedClass]->GetUnrealSourceFile().GetDefinedClassParsingInfo(DefinedClass);
-
-		const FString& BaseClassName         = ParsingInfo.GetBaseClassName();
-		const FString& BaseClassNameStripped = GetClassNameWithPrefixRemoved(BaseClassName);
-
-		if (!BaseClassNameStripped.IsEmpty() && !DefinedClass->GetSuperClass())
-		{
-			UClass* FoundBaseClass = FindObject<UClass>(Package, *BaseClassNameStripped);
-
-			if (FoundBaseClass == nullptr)
-			{
-				FoundBaseClass = FindObject<UClass>(ANY_PACKAGE, *BaseClassNameStripped);
-			}
-
-			if (FoundBaseClass == nullptr)
-			{
-				// Don't know its parent class. Raise error.
-				FError::Throwf(TEXT("Couldn't find parent type for '%s' named '%s' in current module (Package: %s) or any other module parsed so far."), *DefinedClass->GetName(), *BaseClassName, *GetNameSafe(Package));
-			}
-
-			DefinedClass->SetSuperStruct(FoundBaseClass);
-			DefinedClass->ClassCastFlags |= FoundBaseClass->ClassCastFlags;
-		}
-	}
-}
-
-ECompilationResult::Type PreparseModules(const FString& ModuleInfoPath, int32& NumFailures)
-{
-	// Three passes.  1) Public 'Classes' headers (legacy)  2) Public headers   3) Private headers
-	enum EHeaderFolderTypes
-	{
-		PublicClassesHeaders = 0,
-		PublicHeaders = 1,
-		PrivateHeaders,
-
-		FolderType_Count
-	};
-
-	ECompilationResult::Type Result = ECompilationResult::Succeeded;
-
-#if !PLATFORM_EXCEPTIONS_DISABLED
-	FGraphEventArray ExceptionTasks;
-
-	auto LogException = [&Result, &NumFailures, &ExceptionTasks](FString&& Filename, int32 Line, const FString& Message)
-	{
-		auto LogExceptionTask = [&Result, &NumFailures, Filename = MoveTemp(Filename), Line, Message]()
-		{
-			TGuardValue<ELogTimes::Type> DisableLogTimes(GPrintLogTimes, ELogTimes::None);
-
-			FString FormattedErrorMessage = FString::Printf(TEXT("%s(%d): Error: %s\r\n"), *Filename, Line, *Message);
-			Result = ECompilationResult::OtherCompilationError;
-
-			UE_LOG(LogCompile, Log, TEXT("%s"), *FormattedErrorMessage);
-			GWarn->Log(ELogVerbosity::Error, FormattedErrorMessage);
-
-			++NumFailures;
-		};
-
-		if (IsInGameThread())
-		{
-			LogExceptionTask();
-		}
-		else
-		{
-			FGraphEventRef EventRef = FFunctionGraphTask::CreateAndDispatchWhenReady(MoveTemp(LogExceptionTask), TStatId(), nullptr, ENamedThreads::GameThread);
-
-			static FCriticalSection ExceptionCS;
-			FScopeLock Lock(&ExceptionCS);
-			ExceptionTasks.Add(EventRef);
-		}
-	};
-#endif
-
-	for (FManifestModule& Module : GManifest.Modules)
-	{
-		if (Result != ECompilationResult::Succeeded)
-		{
-			break;
-		}
-
-		// Force regeneration of all subsequent modules, otherwise data will get corrupted.
-		Module.ForceRegeneration();
-
-		UPackage* Package = Cast<UPackage>(StaticFindObjectFast(UPackage::StaticClass(), NULL, FName(*Module.LongPackageName), false, false));
-		if (Package == NULL)
-		{
-			Package = CreatePackage(NULL, *Module.LongPackageName);
-		}
-		// Set some package flags for indicating that this package contains script
-		// NOTE: We do this even if we didn't have to create the package, because CoreUObject is compiled into UnrealHeaderTool and we still
-		//       want to make sure our flags get set
-		Package->SetPackageFlags(PKG_ContainsScript | PKG_Compiling);
-		Package->ClearPackageFlags(PKG_ClientOptional | PKG_ServerSideOnly);
-		
-		if(Module.OverrideModuleType == EPackageOverrideType::None)
-		{
-			switch (Module.ModuleType)
-			{
-			case EBuildModuleType::GameEditor:
-			case EBuildModuleType::EngineEditor:
-				Package->SetPackageFlags(PKG_EditorOnly);
-				break;
-
-			case EBuildModuleType::GameDeveloper:
-			case EBuildModuleType::EngineDeveloper:
-				Package->SetPackageFlags(PKG_Developer);
-				break;
-
-			case EBuildModuleType::GameUncooked:
-			case EBuildModuleType::EngineUncooked:
-				Package->SetPackageFlags(PKG_UncookedOnly);
-				break;
-			}
-		}
-		else
-		{
-			// If the user has specified this module to have another package flag, then OR it on
-			switch (Module.OverrideModuleType)
-			{
-			case EPackageOverrideType::EditorOnly:
-				Package->SetPackageFlags(PKG_EditorOnly);
-				break;
-
-			case EPackageOverrideType::EngineDeveloper:
-			case EPackageOverrideType::GameDeveloper:
-				Package->SetPackageFlags(PKG_Developer);
-				break;
-
-			case EPackageOverrideType::EngineUncookedOnly:
-			case EPackageOverrideType::GameUncookedOnly:
-				Package->SetPackageFlags(PKG_UncookedOnly);
-				break;
-			}
-		}
-
-		// Add new module or overwrite whatever we had loaded, that data is obsolete.
-		GPackageToManifestModuleMap.Add(Package, &Module);
-
-		double ThisModulePreparseTime = 0.0;
-		int32 NumHeadersPreparsed = 0;
-		FDurationTimer ThisModuleTimer(ThisModulePreparseTime);
-		ThisModuleTimer.Start();
-
-		// Pre-parse the headers
-		for (int32 PassIndex = 0; PassIndex < FolderType_Count && Result == ECompilationResult::Succeeded; ++PassIndex)
-		{
-			EHeaderFolderTypes CurrentlyProcessing = (EHeaderFolderTypes)PassIndex;
-
-			// We'll make an ordered list of all UObject headers we care about.
-			// @todo uht: Ideally 'dependson' would not be allowed from public -> private, or NOT at all for new style headers
-			const TArray<FString>& UObjectHeaders =
-				(CurrentlyProcessing == PublicClassesHeaders) ? Module.PublicUObjectClassesHeaders :
-				(CurrentlyProcessing == PublicHeaders       ) ? Module.PublicUObjectHeaders        :
-				                                                Module.PrivateUObjectHeaders;
-			if (!UObjectHeaders.Num())
-			{
-				continue;
-			}
-
-			NumHeadersPreparsed += UObjectHeaders.Num();
-
-			TArray<FString> HeaderFiles;
-			HeaderFiles.SetNum(UObjectHeaders.Num());
-
-			{
-				SCOPE_SECONDS_COUNTER_UHT(LoadHeaderContentFromFile);
-				ParallelFor(UObjectHeaders.Num(), [&](int32 Index)
-				{
-					const FString& RawFilename = UObjectHeaders[Index];
-
-#if !PLATFORM_EXCEPTIONS_DISABLED
-					try
-#endif
-					{
-						const FString FullFilename = FPaths::ConvertRelativePathToFull(ModuleInfoPath, RawFilename);
-
-						if (!FFileHelper::LoadFileToString(HeaderFiles[Index], *FullFilename))
-						{
-							FError::Throwf(TEXT("UnrealHeaderTool was unable to load source file '%s'"), *FullFilename);
-						}
-					}
-#if !PLATFORM_EXCEPTIONS_DISABLED
-					catch (TCHAR* ErrorMsg)
-					{
-						FString AbsFilename = IFileManager::Get().ConvertToAbsolutePathForExternalAppForRead(*RawFilename);
-						LogException(MoveTemp(AbsFilename), 1, ErrorMsg);
-					}				
-#endif
-				});
-			}
-
-#if !PLATFORM_EXCEPTIONS_DISABLED
-			FTaskGraphInterface::Get().WaitUntilTasksComplete(ExceptionTasks);
-#endif
-
-			if (Result != ECompilationResult::Succeeded)
-			{
-				continue;
-			}
-
-			TArray<FPerHeaderData> PerHeaderData;
-			PerHeaderData.SetNum(UObjectHeaders.Num());
-
-			ParallelFor(UObjectHeaders.Num(), [&](int32 Index)
-			{
-				const FString& RawFilename = UObjectHeaders[Index];
-
-#if !PLATFORM_EXCEPTIONS_DISABLED
-				try
-#endif
-				{
-					PerformSimplifiedClassParse(Package, *RawFilename, *HeaderFiles[Index], PerHeaderData[Index]);
-				}
-#if !PLATFORM_EXCEPTIONS_DISABLED
-				catch (const FFileLineException& Ex)
-				{
-					FString AbsFilename = IFileManager::Get().ConvertToAbsolutePathForExternalAppForRead(*Ex.Filename);
-					LogException(MoveTemp(AbsFilename), Ex.Line, Ex.Message);
-				}
-				catch (TCHAR* ErrorMsg)
-				{
-					FString AbsFilename = IFileManager::Get().ConvertToAbsolutePathForExternalAppForRead(*RawFilename);
-					LogException(MoveTemp(AbsFilename), 1, ErrorMsg);
-				}
-#endif
-			});
-
-#if !PLATFORM_EXCEPTIONS_DISABLED
-			FTaskGraphInterface::Get().WaitUntilTasksComplete(ExceptionTasks);
-#endif
-
-			if (Result != ECompilationResult::Succeeded)
-			{
-				continue;
-			}
-
-			for (int32 Index = 0; Index < UObjectHeaders.Num(); ++Index)
-			{
-				const FString& RawFilename = UObjectHeaders[Index];
-
-#if !PLATFORM_EXCEPTIONS_DISABLED
-				try
-#endif
-				{
-					// Import class.
-					const FString FullFilename = FPaths::ConvertRelativePathToFull(ModuleInfoPath, RawFilename);
-
-					ProcessInitialClassParse(PerHeaderData[Index]);
-					TSharedRef<FUnrealSourceFile> UnrealSourceFile = PerHeaderData[Index].UnrealSourceFile.ToSharedRef();
-					FUnrealSourceFile* UnrealSourceFilePtr = &UnrealSourceFile.Get();
-					GUnrealSourceFilesMap.Add(FPaths::GetCleanFilename(RawFilename), UnrealSourceFile);
-
-					if (CurrentlyProcessing == PublicClassesHeaders)
-					{
-						GPublicSourceFileSet.Add(UnrealSourceFilePtr);
-					}
-
-					// Save metadata for the class path, both for it's include path and relative to the module base directory
-					if (FullFilename.StartsWith(Module.BaseDirectory))
-					{
-						// Get the path relative to the module directory
-						const TCHAR* ModuleRelativePath = *FullFilename + Module.BaseDirectory.Len();
-
-						UnrealSourceFile->SetModuleRelativePath(ModuleRelativePath);
-
-						// Calculate the include path
-						const TCHAR* IncludePath = ModuleRelativePath;
-
-						// Walk over the first potential slash
-						if (*IncludePath == TEXT('/'))
-						{
-							IncludePath++;
-						}
-
-						// Does this module path start with a known include path location? If so, we can cut that part out of the include path
-						static const TCHAR PublicFolderName[]  = TEXT("Public/");
-						static const TCHAR PrivateFolderName[] = TEXT("Private/");
-						static const TCHAR ClassesFolderName[] = TEXT("Classes/");
-						if (FCString::Strnicmp(IncludePath, PublicFolderName, UE_ARRAY_COUNT(PublicFolderName) - 1) == 0)
-						{
-							IncludePath += (UE_ARRAY_COUNT(PublicFolderName) - 1);
-						}
-						else if (FCString::Strnicmp(IncludePath, PrivateFolderName, UE_ARRAY_COUNT(PrivateFolderName) - 1) == 0)
-						{
-							IncludePath += (UE_ARRAY_COUNT(PrivateFolderName) - 1);
-						}
-						else if (FCString::Strnicmp(IncludePath, ClassesFolderName, UE_ARRAY_COUNT(ClassesFolderName) - 1) == 0)
-						{
-							IncludePath += (UE_ARRAY_COUNT(ClassesFolderName) - 1);
-						}
-
-						// Add the include path
-						if (*IncludePath != 0)
-						{
-							UnrealSourceFile->SetIncludePath(MoveTemp(IncludePath));
-						}
-					}
-				}
-#if !PLATFORM_EXCEPTIONS_DISABLED
-				catch (const FFileLineException& Ex)
-				{
-					FString AbsFilename = IFileManager::Get().ConvertToAbsolutePathForExternalAppForRead(*Ex.Filename);
-					LogException(MoveTemp(AbsFilename), Ex.Line, Ex.Message);
-				}
-				catch (TCHAR* ErrorMsg)
-				{
-					FString AbsFilename = IFileManager::Get().ConvertToAbsolutePathForExternalAppForRead(*RawFilename);
-					LogException(MoveTemp(AbsFilename), 1, ErrorMsg);
-				}
-#endif
-			}
-			if (Result == ECompilationResult::Succeeded && NumFailures != 0)
-			{
-				Result = ECompilationResult::OtherCompilationError;
-			}
-		}
-
-		// Don't resolve superclasses for module when loading from makefile.
-		// Data is only partially loaded at this point.
-#if !PLATFORM_EXCEPTIONS_DISABLED
-		try
-#endif
-		{
-			ResolveSuperClasses(Package);
-		}
-#if !PLATFORM_EXCEPTIONS_DISABLED
-		catch (TCHAR* ErrorMsg)
-		{
-			TGuardValue<ELogTimes::Type> DisableLogTimes(GPrintLogTimes, ELogTimes::None);
-
-			FString FormattedErrorMessage = FString::Printf(TEXT("Error: %s\r\n"), ErrorMsg);
-
-			Result = GCompilationResult;
-
-			UE_LOG(LogCompile, Log, TEXT("%s"), *FormattedErrorMessage);
-			GWarn->Log(ELogVerbosity::Error, FormattedErrorMessage);
-
-			++NumFailures;
-		}
-#endif
-
-		ThisModuleTimer.Stop();
-		UE_LOG(LogCompile, Log, TEXT("Preparsed module %s containing %i files(s) in %.2f secs."), *Module.LongPackageName, NumHeadersPreparsed, ThisModulePreparseTime);
-	}
-
-	return Result;
-}
-
-ECompilationResult::Type UnrealHeaderTool_Main(const FString& ModuleInfoFilename)
-{
-	double MainTime = 0.0;
-	FDurationTimer MainTimer(MainTime);
-	MainTimer.Start();
-
-	check(GIsUCCMakeStandaloneHeaderGenerator);
-	ECompilationResult::Type Result = ECompilationResult::Succeeded;
-
-	FString ModuleInfoPath = FPaths::GetPath(ModuleInfoFilename);
-
-	// Load the manifest file, giving a list of all modules to be processed, pre-sorted by dependency ordering
-#if !PLATFORM_EXCEPTIONS_DISABLED
-	try
-#endif
-	{
-		GManifest = FManifest::LoadFromFile(ModuleInfoFilename);
-	}
-#if !PLATFORM_EXCEPTIONS_DISABLED
-	catch (const TCHAR* Ex)
-	{
-		UE_LOG(LogCompile, Error, TEXT("Failed to load manifest file '%s': %s"), *ModuleInfoFilename, Ex);
-		return GCompilationResult;
-	}
-#endif
-
-	// Counters.
-	int32 NumFailures = 0;
-	double TotalModulePreparseTime = 0.0;
-	double TotalParseAndCodegenTime = 0.0;
-
-	{
-		FDurationTimer TotalModulePreparseTimer(TotalModulePreparseTime);
-		TotalModulePreparseTimer.Start();
-		Result = PreparseModules(ModuleInfoPath, NumFailures);
-		TotalModulePreparseTimer.Stop();
-	}
-	// Do the actual parse of the headers and generate for them
-	if (Result == ECompilationResult::Succeeded)
-	{
-		FScopedDurationTimer ParseAndCodeGenTimer(TotalParseAndCodegenTime);
-
-		TMap<UPackage*, TArray<UClass*>> ClassesByPackageMap;
-		ClassesByPackageMap.Reserve(GManifest.Modules.Num());
-
-		// Verify that all script declared superclasses exist.
-		for (UClass* ScriptClass : TObjectRange<UClass>())
-		{
-			ClassesByPackageMap.FindOrAdd(ScriptClass->GetOutermost()).Add(ScriptClass);
-
-			const UClass* ScriptSuperClass = ScriptClass->GetSuperClass();
-
-			if (ScriptSuperClass && !ScriptSuperClass->HasAnyClassFlags(CLASS_Intrinsic) && GTypeDefinitionInfoMap.Contains(ScriptClass) && !GTypeDefinitionInfoMap.Contains(ScriptSuperClass))
-			{
-				class FSuperClassContextSupplier : public FContextSupplier
-				{
-				public:
-					FSuperClassContextSupplier(const UClass* Class)
-						: DefinitionInfo(GTypeDefinitionInfoMap[Class])
-					{ }
-
-					virtual FString GetContext() override
-					{
-						FString Filename = IFileManager::Get().ConvertToAbsolutePathForExternalAppForRead(*DefinitionInfo->GetUnrealSourceFile().GetFilename());
-						int32 LineNumber = DefinitionInfo->GetLineNumber();
-						return FString::Printf(TEXT("%s(%i)"), *Filename, LineNumber);
-					}
-				private:
-					TSharedRef<FUnrealTypeDefinitionInfo> DefinitionInfo;
-				} ContextSupplier(ScriptClass);
-
-				FContextSupplier* OldContext = GWarn->GetContext();
-
-				TGuardValue<ELogTimes::Type> DisableLogTimes(GPrintLogTimes, ELogTimes::None);
-
-				GWarn->SetContext(&ContextSupplier);
-				GWarn->Log(ELogVerbosity::Error, FString::Printf(TEXT("Error: Superclass %s of class %s not found"), *ScriptSuperClass->GetName(), *ScriptClass->GetName()));
-				GWarn->SetContext(OldContext);
-
-				Result = ECompilationResult::OtherCompilationError;
-				++NumFailures;
-			}
-		}
-
-		if (Result == ECompilationResult::Succeeded)
-		{
-			TArray<IScriptGeneratorPluginInterface*> ScriptPlugins;
-			// Can only export scripts for game targets
-			if (GManifest.IsGameTarget)
-			{
-				GetScriptPlugins(ScriptPlugins);
-			}
-
-			for (const FManifestModule& Module : GManifest.Modules)
-			{
-				if (UPackage* Package = Cast<UPackage>(StaticFindObjectFast(UPackage::StaticClass(), NULL, FName(*Module.LongPackageName), false, false)))
-				{
-					FClasses AllClasses(ClassesByPackageMap.Find(Package));
-					AllClasses.Validate();
-
-					Result = FHeaderParser::ParseAllHeadersInside(AllClasses, GWarn, Package, Module, ScriptPlugins);
-					if (Result != ECompilationResult::Succeeded)
-					{
-						++NumFailures;
-						break;
-					}
-				}
-			}
-
-			{
-				FScopedDurationTimer PluginTimeTracker(GPluginOverheadTime);
-				for (IScriptGeneratorPluginInterface* ScriptGenerator : ScriptPlugins)
-				{
-					ScriptGenerator->FinishExport();
-				}
-			}
-
-			// Get a list of external dependencies from each enabled plugin
-			FString ExternalDependencies;
-			for (IScriptGeneratorPluginInterface* ScriptPlugin : ScriptPlugins)
-			{
-				TArray<FString> PluginExternalDependencies;
-				ScriptPlugin->GetExternalDependencies(PluginExternalDependencies);
-
-				for (const FString& PluginExternalDependency : PluginExternalDependencies)
-				{
-					ExternalDependencies += PluginExternalDependency + LINE_TERMINATOR;
-				}
-			}
-			FFileHelper::SaveStringToFile(ExternalDependencies, *GManifest.ExternalDependenciesFile);
-		}
-	}
-
-	// Avoid TArray slack for meta data.
-	GScriptHelper.Shrink();
-
-	// Finish all async file tasks before stopping the clock
-	FTaskGraphInterface::Get().WaitUntilTasksComplete(GAsyncFileTasks);
-
-	MainTimer.Stop();
-
-	UE_LOG(LogCompile, Log, TEXT("Preparsing %i modules took %.2f seconds"), GManifest.Modules.Num(), TotalModulePreparseTime);
-	UE_LOG(LogCompile, Log, TEXT("Parsing took %.2f seconds"), TotalParseAndCodegenTime - GHeaderCodeGenTime);
-	UE_LOG(LogCompile, Log, TEXT("Code generation took %.2f seconds"), GHeaderCodeGenTime);
-	UE_LOG(LogCompile, Log, TEXT("ScriptPlugin overhead was %.2f seconds"), GPluginOverheadTime);
-	UE_LOG(LogCompile, Log, TEXT("Macroize time was %.2f seconds"), GMacroizeTime);
-
-	FUnrealHeaderToolStats& Stats = FUnrealHeaderToolStats::Get();
-	for (const TPair<FName, double>& Pair : Stats.Counters)
-	{
-		FString CounterName = Pair.Key.ToString();
-		UE_LOG(LogCompile, Log, TEXT("%s timer was %.3f seconds"), *CounterName, Pair.Value);
-	}
-
-	UE_LOG(LogCompile, Log, TEXT("Total time was %.2f seconds"), MainTime);
-
-	if (bWriteContents)
-	{
-		UE_LOG(LogCompile, Log, TEXT("********************************* Wrote reference generated code to ReferenceGeneratedCode."));
-	}
-	else if (bVerifyContents)
-	{
-		UE_LOG(LogCompile, Log, TEXT("********************************* Wrote generated code to VerifyGeneratedCode and compared to ReferenceGeneratedCode"));
-		for (FString& Msg : ChangeMessages)
-		{
-			UE_LOG(LogCompile, Error, TEXT("%s"), *Msg);
-		}
-		TArray<FString> RefFileNames;
-		IFileManager::Get().FindFiles( RefFileNames, *(FPaths::ProjectSavedDir() / TEXT("ReferenceGeneratedCode/*.*")), true, false );
-		TArray<FString> VerFileNames;
-		IFileManager::Get().FindFiles( VerFileNames, *(FPaths::ProjectSavedDir() / TEXT("VerifyGeneratedCode/*.*")), true, false );
-		if (RefFileNames.Num() != VerFileNames.Num())
-		{
-			UE_LOG(LogCompile, Error, TEXT("Number of generated files mismatch ref=%d, ver=%d"), RefFileNames.Num(), VerFileNames.Num());
-		}
-	}
-
-	RequestEngineExit(TEXT("UnrealHeaderTool finished"));
-
-	if (Result != ECompilationResult::Succeeded || NumFailures > 0)
-	{
-		return ECompilationResult::OtherCompilationError;
-	}
-
-	return Result;
-}
-
-UClass* ProcessParsedClass(bool bClassIsAnInterface, TArray<FHeaderProvider>& DependentOn, const FString& ClassName, const FString& BaseClassName, UObject* InParent, EObjectFlags Flags)
-{
-	FString ClassNameStripped = GetClassNameWithPrefixRemoved(*ClassName);
-
-	// All classes must start with a valid unreal prefix
-	if (!FHeaderParser::ClassNameHasValidPrefix(ClassName, ClassNameStripped))
-	{
-		FError::Throwf(TEXT("Invalid class name '%s'. The class name must have an appropriate prefix added (A for Actors, U for other classes)."), *ClassName);
-	}
-
-	if(FHeaderParser::IsReservedTypeName(ClassNameStripped))
-	{
-		FError::Throwf(TEXT("Invalid class name '%s'. Cannot use a reserved name ('%s')."), *ClassName, *ClassNameStripped);
-	}
-
-	// Ensure the base class has any valid prefix and exists as a valid class. Checking for the 'correct' prefix will occur during compilation
-	FString BaseClassNameStripped;
-	if (!BaseClassName.IsEmpty())
-	{
-		BaseClassNameStripped = GetClassNameWithPrefixRemoved(BaseClassName);
-		if (!FHeaderParser::ClassNameHasValidPrefix(BaseClassName, BaseClassNameStripped))
-		{
-			FError::Throwf(TEXT("No prefix or invalid identifier for base class %s.\nClass names must match Unreal prefix specifications (e.g., \"UObject\" or \"AActor\")"), *BaseClassName);
-		}
-	}
-
-	//UE_LOG(LogCompile, Log, TEXT("Class: %s extends %s"),*ClassName,*BaseClassName);
-	// Handle failure and non-class headers.
-	if (BaseClassName.IsEmpty() && (ClassName != TEXT("UObject")))
-	{
-		FError::Throwf(TEXT("Class '%s' must inherit UObject or a UObject-derived class"), *ClassName);
-	}
-
-	if (ClassName == BaseClassName)
-	{
-		FError::Throwf(TEXT("Class '%s' cannot inherit from itself"), *ClassName);
-	}
-
-	// In case the file system and the class disagree on the case of the
-	// class name replace the fname with the one from the script class file
-	// This is needed because not all source control systems respect the
-	// original filename's case
-	FName ClassNameReplace(*ClassName, FNAME_Replace_Not_Safe_For_Threading);
-
-	// Use stripped class name for processing and replace as we did above
-	FName ClassNameStrippedReplace(*ClassNameStripped, FNAME_Replace_Not_Safe_For_Threading);
-
-	UClass* ResultClass = FindObject<UClass>(InParent, *ClassNameStripped);
-
-	// if we aren't generating headers, then we shouldn't set misaligned object, since it won't get cleared
-
-	const static bool bVerboseOutput = FParse::Param(FCommandLine::Get(), TEXT("VERBOSE"));
-
-	if (ResultClass == nullptr || !ResultClass->IsNative())
-	{
-		// detect if the same class name is used in multiple packages
-		if (ResultClass == nullptr)
-		{
-			UClass* ConflictingClass = FindObject<UClass>(ANY_PACKAGE, *ClassNameStripped, true);
-			if (ConflictingClass != nullptr)
-			{
-				UE_LOG_WARNING_UHT(TEXT("Duplicate class name: %s also exists in file %s"), *ClassName, *ConflictingClass->GetOutermost()->GetName());
-			}
-		}
-
-		// Create new class.
-		ResultClass = new(EC_InternalUseOnlyConstructor, InParent, *ClassNameStripped, Flags) UClass(FObjectInitializer(), nullptr);
-
-		// add CLASS_Interface flag if the class is an interface
-		// NOTE: at this pre-parsing/importing stage, we cannot know if our super class is an interface or not,
-		// we leave the validation to the main header parser
-		if (bClassIsAnInterface)
-		{
-			ResultClass->ClassFlags |= CLASS_Interface;
-		}
-
-		if (bVerboseOutput)
-		{
-			UE_LOG(LogCompile, Log, TEXT("Imported: %s"), *ResultClass->GetFullName());
-		}
-	}
-
-	if (bVerboseOutput)
-	{
-		for (const FHeaderProvider& Dependency : DependentOn)
-		{
-			UE_LOG(LogCompile, Log, TEXT("\tAdding %s as a dependency"), *Dependency.ToString());
-		}
-	}
-
-	return ResultClass;
-}
-
-void PerformSimplifiedClassParse(UPackage* InParent, const TCHAR* FileName, const TCHAR* Buffer, FPerHeaderData& PerHeaderData)
-{
-	// Parse the header to extract the information needed
-	FUHTStringBuilder ClassHeaderTextStrippedOfCppText;
-
-	FHeaderParser::SimplifiedClassParse(FileName, Buffer, /*out*/ PerHeaderData.ParsedClassArray, /*out*/ PerHeaderData.DependsOn, ClassHeaderTextStrippedOfCppText);
-
-	FUnrealSourceFile* UnrealSourceFilePtr = new FUnrealSourceFile(InParent, FileName, MoveTemp(ClassHeaderTextStrippedOfCppText));
-	PerHeaderData.UnrealSourceFile = MakeShareable(UnrealSourceFilePtr);
-}
-
-void ProcessInitialClassParse(FPerHeaderData& PerHeaderData)
-{
-	TSharedRef<FUnrealSourceFile> UnrealSourceFile = PerHeaderData.UnrealSourceFile.ToSharedRef();
-	UPackage* InParent = UnrealSourceFile->GetPackage();
-	for (FSimplifiedParsingClassInfo& ParsedClassInfo : PerHeaderData.ParsedClassArray)
-	{
-		UClass* ResultClass = ProcessParsedClass(ParsedClassInfo.IsInterface(), PerHeaderData.DependsOn, ParsedClassInfo.GetClassName(), ParsedClassInfo.GetBaseClassName(), InParent, RF_Public | RF_Standalone);
-		GStructToSourceLine.Add(ResultClass, MakeTuple(UnrealSourceFile, ParsedClassInfo.GetClassDefLine()));
-
-		FScope::AddTypeScope(ResultClass, &UnrealSourceFile->GetScope().Get());
-
-		GTypeDefinitionInfoMap.Add(ResultClass, MakeShared<FUnrealTypeDefinitionInfo>(UnrealSourceFile.Get(), ParsedClassInfo.GetClassDefLine()));
-		UnrealSourceFile->AddDefinedClass(ResultClass, MoveTemp(ParsedClassInfo));
-	}
-
-	for (FHeaderProvider& DependsOnElement : PerHeaderData.DependsOn)
-	{
-		UnrealSourceFile->GetIncludes().AddUnique(DependsOnElement);
-	}
-}
diff --git a/Engine/Source/Programs/UHTLite/Private/HeaderParser.cpp b/Engine/Source/Programs/UHTLite/Private/HeaderParser.cpp
--- a/Engine/Source/Programs/UHTLite/Private/HeaderParser.cpp
+++ b/Engine/Source/Programs/UHTLite/Private/HeaderParser.cpp
@@ -105,14 +105,8 @@ const FName FHeaderParserNames::NAME_SparseClassDataTypes(TEXT("SparseClassDataT
 const FName FHeaderParserNames::NAME_IsConversionRoot(TEXT("IsConversionRoot"));
 const FName FHeaderParserNames::NAME_AdvancedClassDisplay(TEXT("AdvancedClassDisplay"));
 
-EGeneratedCodeVersion FHeaderParser::DefaultGeneratedCodeVersion = EGeneratedCodeVersion::V1;
-TArray<FString> FHeaderParser::StructsWithNoPrefix;
-TArray<FString> FHeaderParser::StructsWithTPrefix;
-FRigVMStructMap FHeaderParser::StructRigVMMap;
-TArray<FString> FHeaderParser::DelegateParameterCountStrings;
-TMap<FString, FString> FHeaderParser::TypeRedirectMap;
-TArray<FString> FHeaderParser::PropertyCPPTypesRequiringUIRanges = { TEXT("float"), TEXT("double") };
-TArray<FString> FHeaderParser::ReservedTypeNames = { TEXT("none") };
+const TArray<FString> FHeaderParser::PropertyCPPTypesRequiringUIRanges = { TEXT("float"), TEXT("double") };
+const TArray<FString> FHeaderParser::ReservedTypeNames = { TEXT("none") };
 TMap<UClass*, ClassDefinitionRange> ClassDefinitionRanges;
 
 /**
@@ -655,312 +649,6 @@ namespace
 		return nullptr;
 	}
 
-	// Check to see if anything in the class hierarchy passed in has CLASS_DefaultToInstanced
-	bool DoesAnythingInHierarchyHaveDefaultToInstanced(UClass* TestClass)
-	{
-		bool bDefaultToInstanced = false;
-
-		UClass* Search = TestClass;
-		while (!bDefaultToInstanced && (Search != NULL))
-		{
-			bDefaultToInstanced = Search->HasAnyClassFlags(CLASS_DefaultToInstanced);
-			if (!bDefaultToInstanced && !Search->HasAnyClassFlags(CLASS_Intrinsic | CLASS_Parsed))
-			{
-				// The class might not have been parsed yet, look for declaration data.
-				if (FClassDeclarationMetaData* ClassDeclarationDataPtr = GClassDeclarations.Find(Search->GetFName()))
-				{
-					bDefaultToInstanced = !!(ClassDeclarationDataPtr->ClassFlags & CLASS_DefaultToInstanced);
-				}
-			}
-			Search = Search->GetSuperClass();
-		}
-
-		return bDefaultToInstanced;
-	}
-
-	FProperty* CreateVariableProperty(FPropertyBase& VarProperty, FFieldVariant Scope, const FName& Name, EObjectFlags ObjectFlags, EVariableCategory::Type VariableCategory, FUnrealSourceFile* UnrealSourceFile)
-	{
-		// Check if it's an enum class property
-		if (const EUnderlyingEnumType* EnumPropType = GEnumUnderlyingTypes.Find(VarProperty.Enum))
-		{
-			FPropertyBase UnderlyingProperty = VarProperty;
-			UnderlyingProperty.Enum = nullptr;
-			switch (*EnumPropType)
-			{
-				case EUnderlyingEnumType::int8:        UnderlyingProperty.Type = CPT_Int8;   break;
-				case EUnderlyingEnumType::int16:       UnderlyingProperty.Type = CPT_Int16;  break;
-				case EUnderlyingEnumType::int32:       UnderlyingProperty.Type = CPT_Int;    break;
-				case EUnderlyingEnumType::int64:       UnderlyingProperty.Type = CPT_Int64;  break;
-				case EUnderlyingEnumType::uint8:       UnderlyingProperty.Type = CPT_Byte;   break;
-				case EUnderlyingEnumType::uint16:      UnderlyingProperty.Type = CPT_UInt16; break;
-				case EUnderlyingEnumType::uint32:      UnderlyingProperty.Type = CPT_UInt32; break;
-				case EUnderlyingEnumType::uint64:      UnderlyingProperty.Type = CPT_UInt64; break;
-				case EUnderlyingEnumType::Unspecified: UnderlyingProperty.Type = CPT_Int;    break;
-
-				default:
-					check(false);
-			}
-
-			if (*EnumPropType == EUnderlyingEnumType::Unspecified)
-			{
-				UnderlyingProperty.IntType = EIntType::Unsized;
-			}
-
-			FEnumProperty* Result = new FEnumProperty(Scope, Name, ObjectFlags);
-			FNumericProperty* UnderlyingProp = CastFieldChecked<FNumericProperty>(CreateVariableProperty(UnderlyingProperty, Result, TEXT("UnderlyingType"), ObjectFlags, VariableCategory, UnrealSourceFile));
-			Result->UnderlyingProp = MoveTemp(UnderlyingProp);
-			Result->Enum = VarProperty.Enum;
-
-			return Result;
-		}
-
-		switch (VarProperty.Type)
-		{
-			case CPT_Byte:
-			{
-				FByteProperty* Result = new FByteProperty(Scope, Name, ObjectFlags);
-				Result->Enum = VarProperty.Enum;
-				check(VarProperty.IntType == EIntType::Sized);
-				return Result;
-			}
-
-			case CPT_Int8:
-			{
-				FInt8Property* Result = new FInt8Property(Scope, Name, ObjectFlags);
-				check(VarProperty.IntType == EIntType::Sized);
-				return Result;
-			}
-
-			case CPT_Int16:
-			{
-				FInt16Property* Result = new FInt16Property(Scope, Name, ObjectFlags);
-				check(VarProperty.IntType == EIntType::Sized);
-				return Result;
-			}
-
-			case CPT_Int:
-			{
-				FIntProperty* Result = new FIntProperty(Scope, Name, ObjectFlags);
-				if (VarProperty.IntType == EIntType::Unsized)
-				{
-					GUnsizedProperties.Add(Result);
-				}
-				return Result;
-			}
-
-			case CPT_Int64:
-			{
-				FInt64Property* Result = new FInt64Property(Scope, Name, ObjectFlags);
-				check(VarProperty.IntType == EIntType::Sized);
-				return Result;
-			}
-
-			case CPT_UInt16:
-			{
-				FUInt16Property* Result = new FUInt16Property(Scope, Name, ObjectFlags);
-				check(VarProperty.IntType == EIntType::Sized);
-				return Result;
-			}
-
-			case CPT_UInt32:
-			{
-				FUInt32Property* Result = new FUInt32Property(Scope, Name, ObjectFlags);
-				if (VarProperty.IntType == EIntType::Unsized)
-				{
-					GUnsizedProperties.Add(Result);
-				}
-				return Result;
-			}
-
-			case CPT_UInt64:
-			{
-				FUInt64Property* Result = new FUInt64Property(Scope, Name, ObjectFlags);
-				check(VarProperty.IntType == EIntType::Sized);
-				return Result;
-			}
-
-			case CPT_Bool:
-			{
-				FBoolProperty* Result = new FBoolProperty(Scope, Name, ObjectFlags);
-				Result->SetBoolSize(sizeof(bool), true);
-				return Result;
-			}
-
-			case CPT_Bool8:
-			{
-				FBoolProperty* Result = new FBoolProperty(Scope, Name, ObjectFlags);
-				Result->SetBoolSize((VariableCategory == EVariableCategory::Return) ? sizeof(bool) : sizeof(uint8), VariableCategory == EVariableCategory::Return);
-				return Result;
-			}
-
-			case CPT_Bool16:
-			{
-				FBoolProperty* Result = new FBoolProperty(Scope, Name, ObjectFlags);
-				Result->SetBoolSize((VariableCategory == EVariableCategory::Return) ? sizeof(bool) : sizeof(uint16), VariableCategory == EVariableCategory::Return);
-				return Result;
-			}
-
-			case CPT_Bool32:
-			{
-				FBoolProperty* Result = new FBoolProperty(Scope, Name, ObjectFlags);
-				Result->SetBoolSize((VariableCategory == EVariableCategory::Return) ? sizeof(bool) : sizeof(uint32), VariableCategory == EVariableCategory::Return);
-				return Result;
-			}
-
-			case CPT_Bool64:
-			{
-				FBoolProperty* Result = new FBoolProperty(Scope, Name, ObjectFlags);
-				Result->SetBoolSize((VariableCategory == EVariableCategory::Return) ? sizeof(bool) : sizeof(uint64), VariableCategory == EVariableCategory::Return);
-				return Result;
-			}
-
-			case CPT_Float:
-			{
-				FFloatProperty* Result = new FFloatProperty(Scope, Name, ObjectFlags);
-				return Result;
-			}
-
-			case CPT_Double:
-			{
-				FDoubleProperty* Result = new FDoubleProperty(Scope, Name, ObjectFlags);
-				return Result;
-			}
-
-			case CPT_ObjectReference:
-				check(VarProperty.PropertyClass);
-
-				if (VarProperty.PropertyClass->IsChildOf(UClass::StaticClass()))
-				{
-					FClassProperty* Result = new FClassProperty(Scope, Name, ObjectFlags);
-					Result->MetaClass     = VarProperty.MetaClass;
-					Result->PropertyClass = VarProperty.PropertyClass;
-					return Result;
-				}
-				else
-				{
-					if (DoesAnythingInHierarchyHaveDefaultToInstanced(VarProperty.PropertyClass))
-					{
-						VarProperty.PropertyFlags |= CPF_InstancedReference;
-						AddEditInlineMetaData(VarProperty.MetaData);
-					}
-
-					FObjectProperty* Result = new FObjectProperty(Scope, Name, ObjectFlags);
-					Result->PropertyClass = VarProperty.PropertyClass;
-					return Result;
-				}
-
-			case CPT_WeakObjectReference:
-			{
-				check(VarProperty.PropertyClass);
-
-				FWeakObjectProperty* Result = new FWeakObjectProperty(Scope, Name, ObjectFlags);
-				Result->PropertyClass = VarProperty.PropertyClass;
-				return Result;
-			}
-
-			case CPT_LazyObjectReference:
-			{
-				check(VarProperty.PropertyClass);
-
-				FLazyObjectProperty* Result = new FLazyObjectProperty(Scope, Name, ObjectFlags);
-				Result->PropertyClass = VarProperty.PropertyClass;
-				return Result;
-			}
-
-			case CPT_SoftObjectReference:
-				check(VarProperty.PropertyClass);
-
-				if (VarProperty.PropertyClass->IsChildOf(UClass::StaticClass()))
-				{
-					FSoftClassProperty* Result = new FSoftClassProperty(Scope, Name, ObjectFlags);
-					Result->MetaClass     = VarProperty.MetaClass;
-					Result->PropertyClass = VarProperty.PropertyClass;
-					return Result;
-				}
-				else
-				{
-					FSoftObjectProperty* Result = new FSoftObjectProperty(Scope, Name, ObjectFlags);
-					Result->PropertyClass = VarProperty.PropertyClass;
-					return Result;
-				}
-
-			case CPT_Interface:
-			{
-				check(VarProperty.PropertyClass);
-				check(VarProperty.PropertyClass->HasAnyClassFlags(CLASS_Interface));
-
-				FInterfaceProperty* Result = new  FInterfaceProperty(Scope, Name, ObjectFlags);
-				Result->InterfaceClass = VarProperty.PropertyClass;
-				return Result;
-			}
-
-			case CPT_Name:
-			{
-				FNameProperty* Result = new FNameProperty(Scope, Name, ObjectFlags);
-				return Result;
-			}
-
-			case CPT_String:
-			{
-				FStrProperty* Result = new FStrProperty(Scope, Name, ObjectFlags);
-				return Result;
-			}
-
-			case CPT_Text:
-			{
-				FTextProperty* Result = new FTextProperty(Scope, Name, ObjectFlags);
-				return Result;
-			}
-
-			case CPT_Struct:
-			{
-				if (VarProperty.Struct->StructFlags & STRUCT_HasInstancedReference)
-				{
-					VarProperty.PropertyFlags |= CPF_ContainsInstancedReference;
-				}
-
-				FStructProperty* Result = new FStructProperty(Scope, Name, ObjectFlags);
-				Result->Struct = VarProperty.Struct;
-				return Result;
-			}
-
-			case CPT_Delegate:
-			{
-				FDelegateProperty* Result = new FDelegateProperty(Scope, Name, ObjectFlags);
-				return Result;
-			}
-
-			case CPT_MulticastDelegate:
-			{
-				FMulticastDelegateProperty* Result;
-				if (VarProperty.Function->IsA<USparseDelegateFunction>())
-				{
-					Result = new FMulticastSparseDelegateProperty(Scope, Name, ObjectFlags);
-				}
-				else
-				{
-					Result = new FMulticastInlineDelegateProperty(Scope, Name, ObjectFlags);
-				}
-				return Result;
-			}
-
-			case CPT_FieldPath:
-			{
-				FFieldPathProperty* Result = new FFieldPathProperty(Scope, Name, ObjectFlags);
-				Result->PropertyClass = VarProperty.PropertyPathClass;
-
-				return Result;
-			}
-
-			default:
-				FError::Throwf(TEXT("Unknown property type %i"), (uint8)VarProperty.Type);
-		}
-
-		// Unreachable
-		check(false); //-V779
-		return nullptr;
-	}
-
 	/**
 	 * Ensures at script compile time that the metadata formatting is correct
 	 * @param	InKey			the metadata key being added
@@ -1293,6 +981,317 @@ namespace
 	};
 }
 
+// Check to see if anything in the class hierarchy passed in has CLASS_DefaultToInstanced
+bool FHeaderParser::DoesAnythingInHierarchyHaveDefaultToInstanced(UClass* TestClass) const
+{
+	bool bDefaultToInstanced = false;
+
+	UClass* Search = TestClass;
+	while (!bDefaultToInstanced && (Search != NULL))
+	{
+		bDefaultToInstanced = Search->HasAnyClassFlags(CLASS_DefaultToInstanced);
+		if (!bDefaultToInstanced && !Search->HasAnyClassFlags(CLASS_Intrinsic | CLASS_Parsed))
+		{
+			// The class might not have been parsed yet, look for declaration data.
+			if (const FClassDeclarationMetaData* ClassDeclarationDataPtr = ClassDeclarations.Find(Search->GetFName()))
+			{
+				bDefaultToInstanced = !!(ClassDeclarationDataPtr->ClassFlags & CLASS_DefaultToInstanced);
+			}
+		}
+		Search = Search->GetSuperClass();
+	}
+
+	return bDefaultToInstanced;
+}
+
+FProperty* FHeaderParser::CreateVariableProperty(FPropertyBase& VarProperty, FFieldVariant Scope, const FName& Name, EObjectFlags ObjectFlags, EVariableCategory::Type VariableCategory, FUnrealSourceFile* UnrealSourceFile) const
+{
+	{
+		// UHTLite
+		FRWScopeLock Lock(EnumUnderlyingTypesLock, SLT_ReadOnly);
+
+		// Check if it's an enum class property
+		if (const EUnderlyingEnumType* EnumPropType = EnumUnderlyingTypes.Find(VarProperty.Enum))
+		{
+			FPropertyBase UnderlyingProperty = VarProperty;
+			UnderlyingProperty.Enum = nullptr;
+			switch (*EnumPropType)
+			{
+			case EUnderlyingEnumType::int8:        UnderlyingProperty.Type = CPT_Int8;   break;
+			case EUnderlyingEnumType::int16:       UnderlyingProperty.Type = CPT_Int16;  break;
+			case EUnderlyingEnumType::int32:       UnderlyingProperty.Type = CPT_Int;    break;
+			case EUnderlyingEnumType::int64:       UnderlyingProperty.Type = CPT_Int64;  break;
+			case EUnderlyingEnumType::uint8:       UnderlyingProperty.Type = CPT_Byte;   break;
+			case EUnderlyingEnumType::uint16:      UnderlyingProperty.Type = CPT_UInt16; break;
+			case EUnderlyingEnumType::uint32:      UnderlyingProperty.Type = CPT_UInt32; break;
+			case EUnderlyingEnumType::uint64:      UnderlyingProperty.Type = CPT_UInt64; break;
+			case EUnderlyingEnumType::Unspecified: UnderlyingProperty.Type = CPT_Int;    break;
+
+			default:
+				check(false);
+			}
+
+			if (*EnumPropType == EUnderlyingEnumType::Unspecified)
+			{
+				UnderlyingProperty.IntType = EIntType::Unsized;
+			}
+
+			FEnumProperty* Result = new FEnumProperty(Scope, Name, ObjectFlags);
+			FNumericProperty* UnderlyingProp = CastFieldChecked<FNumericProperty>(CreateVariableProperty(UnderlyingProperty, Result, TEXT("UnderlyingType"), ObjectFlags, VariableCategory, UnrealSourceFile));
+			Result->UnderlyingProp = MoveTemp(UnderlyingProp);
+			Result->Enum = VarProperty.Enum;
+
+			return Result;
+		}
+	}
+
+	switch (VarProperty.Type)
+	{
+	case CPT_Byte:
+	{
+		FByteProperty* Result = new FByteProperty(Scope, Name, ObjectFlags);
+		Result->Enum = VarProperty.Enum;
+		check(VarProperty.IntType == EIntType::Sized);
+		return Result;
+	}
+
+	case CPT_Int8:
+	{
+		FInt8Property* Result = new FInt8Property(Scope, Name, ObjectFlags);
+		check(VarProperty.IntType == EIntType::Sized);
+		return Result;
+	}
+
+	case CPT_Int16:
+	{
+		FInt16Property* Result = new FInt16Property(Scope, Name, ObjectFlags);
+		check(VarProperty.IntType == EIntType::Sized);
+		return Result;
+	}
+
+	case CPT_Int:
+	{
+		FIntProperty* Result = new FIntProperty(Scope, Name, ObjectFlags);
+		if (VarProperty.IntType == EIntType::Unsized)
+		{
+			UnsizedProperties.Add(Result);
+		}
+		return Result;
+	}
+
+	case CPT_Int64:
+	{
+		FInt64Property* Result = new FInt64Property(Scope, Name, ObjectFlags);
+		check(VarProperty.IntType == EIntType::Sized);
+		return Result;
+	}
+
+	case CPT_UInt16:
+	{
+		FUInt16Property* Result = new FUInt16Property(Scope, Name, ObjectFlags);
+		check(VarProperty.IntType == EIntType::Sized);
+		return Result;
+	}
+
+	case CPT_UInt32:
+	{
+		FUInt32Property* Result = new FUInt32Property(Scope, Name, ObjectFlags);
+		if (VarProperty.IntType == EIntType::Unsized)
+		{
+			UnsizedProperties.Add(Result);
+		}
+		return Result;
+	}
+
+	case CPT_UInt64:
+	{
+		FUInt64Property* Result = new FUInt64Property(Scope, Name, ObjectFlags);
+		check(VarProperty.IntType == EIntType::Sized);
+		return Result;
+	}
+
+	case CPT_Bool:
+	{
+		FBoolProperty* Result = new FBoolProperty(Scope, Name, ObjectFlags);
+		Result->SetBoolSize(sizeof(bool), true);
+		return Result;
+	}
+
+	case CPT_Bool8:
+	{
+		FBoolProperty* Result = new FBoolProperty(Scope, Name, ObjectFlags);
+		Result->SetBoolSize((VariableCategory == EVariableCategory::Return) ? sizeof(bool) : sizeof(uint8), VariableCategory == EVariableCategory::Return);
+		return Result;
+	}
+
+	case CPT_Bool16:
+	{
+		FBoolProperty* Result = new FBoolProperty(Scope, Name, ObjectFlags);
+		Result->SetBoolSize((VariableCategory == EVariableCategory::Return) ? sizeof(bool) : sizeof(uint16), VariableCategory == EVariableCategory::Return);
+		return Result;
+	}
+
+	case CPT_Bool32:
+	{
+		FBoolProperty* Result = new FBoolProperty(Scope, Name, ObjectFlags);
+		Result->SetBoolSize((VariableCategory == EVariableCategory::Return) ? sizeof(bool) : sizeof(uint32), VariableCategory == EVariableCategory::Return);
+		return Result;
+	}
+
+	case CPT_Bool64:
+	{
+		FBoolProperty* Result = new FBoolProperty(Scope, Name, ObjectFlags);
+		Result->SetBoolSize((VariableCategory == EVariableCategory::Return) ? sizeof(bool) : sizeof(uint64), VariableCategory == EVariableCategory::Return);
+		return Result;
+	}
+
+	case CPT_Float:
+	{
+		FFloatProperty* Result = new FFloatProperty(Scope, Name, ObjectFlags);
+		return Result;
+	}
+
+	case CPT_Double:
+	{
+		FDoubleProperty* Result = new FDoubleProperty(Scope, Name, ObjectFlags);
+		return Result;
+	}
+
+	case CPT_ObjectReference:
+		check(VarProperty.PropertyClass);
+
+		if (VarProperty.PropertyClass->IsChildOf(UClass::StaticClass()))
+		{
+			FClassProperty* Result = new FClassProperty(Scope, Name, ObjectFlags);
+			Result->MetaClass = VarProperty.MetaClass;
+			Result->PropertyClass = VarProperty.PropertyClass;
+			return Result;
+		}
+		else
+		{
+			if (DoesAnythingInHierarchyHaveDefaultToInstanced(VarProperty.PropertyClass))
+			{
+				VarProperty.PropertyFlags |= CPF_InstancedReference;
+				AddEditInlineMetaData(VarProperty.MetaData);
+			}
+
+			FObjectProperty* Result = new FObjectProperty(Scope, Name, ObjectFlags);
+			Result->PropertyClass = VarProperty.PropertyClass;
+			return Result;
+		}
+
+	case CPT_WeakObjectReference:
+	{
+		check(VarProperty.PropertyClass);
+
+		FWeakObjectProperty* Result = new FWeakObjectProperty(Scope, Name, ObjectFlags);
+		Result->PropertyClass = VarProperty.PropertyClass;
+		return Result;
+	}
+
+	case CPT_LazyObjectReference:
+	{
+		check(VarProperty.PropertyClass);
+
+		FLazyObjectProperty* Result = new FLazyObjectProperty(Scope, Name, ObjectFlags);
+		Result->PropertyClass = VarProperty.PropertyClass;
+		return Result;
+	}
+
+	case CPT_SoftObjectReference:
+		check(VarProperty.PropertyClass);
+
+		if (VarProperty.PropertyClass->IsChildOf(UClass::StaticClass()))
+		{
+			FSoftClassProperty* Result = new FSoftClassProperty(Scope, Name, ObjectFlags);
+			Result->MetaClass = VarProperty.MetaClass;
+			Result->PropertyClass = VarProperty.PropertyClass;
+			return Result;
+		}
+		else
+		{
+			FSoftObjectProperty* Result = new FSoftObjectProperty(Scope, Name, ObjectFlags);
+			Result->PropertyClass = VarProperty.PropertyClass;
+			return Result;
+		}
+
+	case CPT_Interface:
+	{
+		check(VarProperty.PropertyClass);
+		check(VarProperty.PropertyClass->HasAnyClassFlags(CLASS_Interface));
+
+		FInterfaceProperty* Result = new  FInterfaceProperty(Scope, Name, ObjectFlags);
+		Result->InterfaceClass = VarProperty.PropertyClass;
+		return Result;
+	}
+
+	case CPT_Name:
+	{
+		FNameProperty* Result = new FNameProperty(Scope, Name, ObjectFlags);
+		return Result;
+	}
+
+	case CPT_String:
+	{
+		FStrProperty* Result = new FStrProperty(Scope, Name, ObjectFlags);
+		return Result;
+	}
+
+	case CPT_Text:
+	{
+		FTextProperty* Result = new FTextProperty(Scope, Name, ObjectFlags);
+		return Result;
+	}
+
+	case CPT_Struct:
+	{
+		if (VarProperty.Struct->StructFlags & STRUCT_HasInstancedReference)
+		{
+			VarProperty.PropertyFlags |= CPF_ContainsInstancedReference;
+		}
+
+		FStructProperty* Result = new FStructProperty(Scope, Name, ObjectFlags);
+		Result->Struct = VarProperty.Struct;
+		return Result;
+	}
+
+	case CPT_Delegate:
+	{
+		FDelegateProperty* Result = new FDelegateProperty(Scope, Name, ObjectFlags);
+		return Result;
+	}
+
+	case CPT_MulticastDelegate:
+	{
+		FMulticastDelegateProperty* Result;
+		if (VarProperty.Function->IsA<USparseDelegateFunction>())
+		{
+			Result = new FMulticastSparseDelegateProperty(Scope, Name, ObjectFlags);
+		}
+		else
+		{
+			Result = new FMulticastInlineDelegateProperty(Scope, Name, ObjectFlags);
+		}
+		return Result;
+	}
+
+	case CPT_FieldPath:
+	{
+		FFieldPathProperty* Result = new FFieldPathProperty(Scope, Name, ObjectFlags);
+		Result->PropertyClass = VarProperty.PropertyPathClass;
+
+		return Result;
+	}
+
+	default:
+		FError::Throwf(TEXT("Unknown property type %i"), (uint8)VarProperty.Type);
+	}
+
+	// Unreachable
+	check(false); //-V779
+	return nullptr;
+}
+
 /////////////////////////////////////////////////////
 // FScriptLocation
 
@@ -1460,9 +1459,12 @@ FField* FHeaderParser::FindProperty(UStruct* Scope, const TCHAR* InIdentifier, b
 
 	return NULL;
 }
+
 /**
  * @return	true if Scope has FProperty objects in its list of fields
  */
+// UHTLite NOTE: Not required for code-gen. Will be refactored later.
+/*
 bool FHeaderParser::HasMemberProperties( const UStruct* Scope )
 {
 	// it's safe to pass a NULL Scope to TFieldIterator, but this function shouldn't be called with a NULL Scope
@@ -1470,6 +1472,7 @@ bool FHeaderParser::HasMemberProperties( const UStruct* Scope )
 	TFieldIterator<FProperty> It(Scope,EFieldIteratorFlags::ExcludeSuper);
 	return It ? true : false;
 }
+*/
 
 /**
  * Get the parent struct specified.
@@ -1479,6 +1482,8 @@ bool FHeaderParser::HasMemberProperties( const UStruct* Scope )
  *
  * @return	a pointer to the parent struct with the specified name, or NULL if the parent couldn't be found
  */
+// UHTLite NOTE: Not required for code-gen. Will be refactored later.
+/*
 UStruct* FHeaderParser::GetSuperScope( UStruct* CurrentScope, const FName& SearchName )
 {
 	UStruct* SuperScope = CurrentScope;
@@ -1506,6 +1511,7 @@ UStruct* FHeaderParser::GetSuperScope( UStruct* CurrentScope, const FName& Searc
 
 	return SuperScope;
 }
+*/
 
 /**
  * Adds source file's include path to given metadata.
@@ -1513,10 +1519,10 @@ UStruct* FHeaderParser::GetSuperScope( UStruct* CurrentScope, const FName& Searc
  * @param Type Type for which to add include path.
  * @param MetaData Meta data to fill the information.
  */
-void AddIncludePathToMetadata(UField* Type, TMap<FName, FString> &MetaData)
+void FHeaderParser::AddIncludePathToMetadata(UField* Type, TMap<FName, FString> &MetaData) const
 {
 	// Add metadata for the include path.
-	TSharedRef<FUnrealTypeDefinitionInfo>* TypeDefinitionPtr = GTypeDefinitionInfoMap.Find(Type);
+	const TSharedRef<FUnrealTypeDefinitionInfo>* TypeDefinitionPtr = TypeDefinitionInfoMap.Find(Type);
 	if (TypeDefinitionPtr != nullptr)
 	{
 		MetaData.Add(NAME_IncludePath, (*TypeDefinitionPtr)->GetUnrealSourceFile().GetIncludePath());
@@ -1529,7 +1535,7 @@ void AddIncludePathToMetadata(UField* Type, TMap<FName, FString> &MetaData)
  * @param SourceFile Given source file.
  * @param MetaData Meta data to fill the information.
  */
-void AddModuleRelativePathToMetadata(FUnrealSourceFile& SourceFile, TMap<FName, FString> &MetaData)
+void FHeaderParser::AddModuleRelativePathToMetadata(FUnrealSourceFile& SourceFile, TMap<FName, FString> &MetaData) const
 {
 	MetaData.Add(NAME_ModuleRelativePath, SourceFile.GetModuleRelativePath());
 }
@@ -1540,10 +1546,10 @@ void AddModuleRelativePathToMetadata(FUnrealSourceFile& SourceFile, TMap<FName,
  * @param Type Type for which to add module's relative path.
  * @param MetaData Meta data to fill the information.
  */
-void AddModuleRelativePathToMetadata(UField* Type, TMap<FName, FString> &MetaData)
+void FHeaderParser::AddModuleRelativePathToMetadata(UField* Type, TMap<FName, FString> &MetaData) const
 {
 	// Add metadata for the module relative path.
-	TSharedRef<FUnrealTypeDefinitionInfo>* TypeDefinitionPtr = GTypeDefinitionInfoMap.Find(Type);
+	const TSharedRef<FUnrealTypeDefinitionInfo>* TypeDefinitionPtr = TypeDefinitionInfoMap.Find(Type);
 	if (TypeDefinitionPtr != nullptr)
 	{
 		MetaData.Add(NAME_ModuleRelativePath, (*TypeDefinitionPtr)->GetUnrealSourceFile().GetModuleRelativePath());
@@ -1645,10 +1651,10 @@ UEnum* FHeaderParser::CompileEnum()
 
 	if (CompilerDirectiveStack.Num() > 0 && (CompilerDirectiveStack.Last() & ECompilerDirective::WithEditorOnlyData) != 0)
 	{
-		GEditorOnlyDataTypes.Add(Enum);
+		EditorOnlyDataTypes.Add(Enum);
 	}
 
-	GTypeDefinitionInfoMap.Add(Enum, MakeShared<FUnrealTypeDefinitionInfo>(*CurrentSrcFile, InputLine));
+	TypeDefinitionInfoMap.Add(Enum, MakeShared<FUnrealTypeDefinitionInfo>(*CurrentSrcFile, InputLine));
 
 	// Validate the metadata for the enum
 	ValidateMetaDataFormat(Enum, EnumToken.MetaData);
@@ -1707,7 +1713,12 @@ UEnum* FHeaderParser::CompileEnum()
 			UnderlyingType = EUnderlyingEnumType::Unspecified;
 		}
 
-		GEnumUnderlyingTypes.Add(Enum, UnderlyingType);
+		{
+			// UHTLite
+			FRWScopeLock Lock(EnumUnderlyingTypesLock, SLT_Write);
+
+			EnumUnderlyingTypes.Add(Enum, UnderlyingType);
+		}
 	}
 
 	if (UnderlyingType != EUnderlyingEnumType::uint8 && EnumToken.MetaData.Contains(NAME_BlueprintType))
@@ -2416,7 +2427,7 @@ UScriptStruct* FHeaderParser::CompileStructDeclaration(FClasses& AllClasses)
 			{
 				StructFlags |= STRUCT_Immutable | STRUCT_Atomic;
 
-				if (!FPaths::IsSamePath(Filename, GTypeDefinitionInfoMap[UObject::StaticClass()]->GetUnrealSourceFile().GetFilename()))
+				if (!FPaths::IsSamePath(Filename, TypeDefinitionInfoMap[UObject::StaticClass()]->GetUnrealSourceFile().GetFilename()))
 				{
 					UE_LOG_ERROR_UHT(TEXT("Immutable is being phased out in favor of SerializeNative, and is only legal on the mirror structs declared in UObject"));
 				}
@@ -2435,6 +2446,8 @@ UScriptStruct* FHeaderParser::CompileStructDeclaration(FClasses& AllClasses)
 
 		if (UStruct* FoundType = FindObject<UStruct>(ANY_PACKAGE, *EffectiveStructName))
 		{
+			// UHTLite NOTE: Not required for code-gen. Will be refactored later.
+			/*
 			if (TTuple<TSharedRef<FUnrealSourceFile>, int32>* FoundTypeInfo = GStructToSourceLine.Find(FoundType))
 			{
 				FError::Throwf(
@@ -2445,6 +2458,7 @@ UScriptStruct* FHeaderParser::CompileStructDeclaration(FClasses& AllClasses)
 				);
 			}
 			else
+			*/
 			{
 				FError::Throwf(TEXT("struct: '%s' conflicts with another type of the same name"), *EffectiveStructName);
 			}
@@ -2573,7 +2587,7 @@ UScriptStruct* FHeaderParser::CompileStructDeclaration(FClasses& AllClasses)
 	UScriptStruct* Struct = new(EC_InternalUseOnlyConstructor, CurrentSrcFile->GetPackage(), *EffectiveStructName, RF_Public) UScriptStruct(FObjectInitializer(), BaseStruct);
 
 	Scope->AddType(Struct);
-	GTypeDefinitionInfoMap.Add(Struct, MakeShared<FUnrealTypeDefinitionInfo>(*CurrentSrcFile, InputLine));
+	TypeDefinitionInfoMap.Add(Struct, MakeShared<FUnrealTypeDefinitionInfo>(*CurrentSrcFile, InputLine));
 	FScope::AddTypeScope(Struct, &CurrentSrcFile->GetScope().Get());
 
 	AddModuleRelativePathToMetadata(Struct, MetaData);
@@ -2887,6 +2901,8 @@ void FHeaderParser::CheckAllow( const TCHAR* Thing, ENestAllowFlags AllowFlags )
 	}
 }
 
+// UHTLite NOTE: Not required for code-gen. Will be refactored later.
+/*
 bool FHeaderParser::AllowReferenceToClass(UStruct* Scope, UClass* CheckClass) const
 {
 	check(CheckClass);
@@ -2895,6 +2911,7 @@ bool FHeaderParser::AllowReferenceToClass(UStruct* Scope, UClass* CheckClass) co
 		|| ((CheckClass->ClassFlags&CLASS_Parsed) != 0)
 		|| ((CheckClass->ClassFlags&CLASS_Intrinsic) != 0);
 }
+*/
 
 /*-----------------------------------------------------------------------------
 	Nest management.
@@ -4622,7 +4639,10 @@ void FHeaderParser::GetVarType(
 
 		if (VariableCategory == EVariableCategory::Member)
 		{
-			EUnderlyingEnumType* EnumUnderlyingType = GEnumUnderlyingTypes.Find(Enum);
+			// UHTLite
+			FRWScopeLock Lock(EnumUnderlyingTypesLock, SLT_ReadOnly);
+
+			EUnderlyingEnumType* EnumUnderlyingType = EnumUnderlyingTypes.Find(Enum);
 			if (!EnumUnderlyingType)
 			{
 				FError::Throwf(TEXT("You cannot use the raw enum name as a type for member variables, instead use TEnumAsByte or a C++11 enum class with an explicit underlying type."), *Enum->CppType);
@@ -5451,7 +5471,7 @@ FProperty* FHeaderParser::GetVarNameAndDim
 
 		if (VarProperty.AllocatorType == EAllocatorType::MemoryImage)
 		{
-			GPropertyUsesMemoryImageAllocator.Add(Array);
+			PropertyUsesMemoryImageAllocator.Add(Array);
 		}
 	}
 	else if (VarProperty.ArrayType == EArrayType::Set)
@@ -5490,7 +5510,7 @@ FProperty* FHeaderParser::GetVarNameAndDim
 
 		if (VarProperty.AllocatorType == EAllocatorType::MemoryImage)
 		{
-			GPropertyUsesMemoryImageAllocator.Add(Map);
+			PropertyUsesMemoryImageAllocator.Add(Map);
 		}
 	}
 	else
@@ -5500,7 +5520,7 @@ FProperty* FHeaderParser::GetVarNameAndDim
 		if (VarProperty.ArrayType == EArrayType::Static)
 		{
 			Result->ArrayDim = 2; // 2 = static array
-			GArrayDimensions.Add(Result, Dimensions.String);
+			ArrayDimensions.Add(Result, Dimensions.String);
 		}
 
 		Result->PropertyFlags = VarProperty.PropertyFlags;
@@ -5679,7 +5699,12 @@ bool FHeaderParser::CompileDeclaration(FClasses& AllClasses, TArray<UDelegateFun
 		bHaveSeenUClass = true;
 		bEncounteredNewStyleClass_UnmatchedBrackets = true;
 		UClass* Class = CompileClassDeclaration(AllClasses);
+
+		// UHTLite NOTE: Not required for code-gen. Will be refactored later.
+		/*
 		GStructToSourceLine.Add(Class, MakeTuple(GetCurrentSourceFile()->AsShared(), Token.StartLine));
+		*/
+
 		return true;
 	}
 
@@ -5731,7 +5756,12 @@ bool FHeaderParser::CompileDeclaration(FClasses& AllClasses, TArray<UDelegateFun
 	{
 		// Struct definition.
 		UScriptStruct* Struct = CompileStructDeclaration(AllClasses);
+
+		// UHTLite NOTE: Not required for code-gen. Will be refactored later.
+		/*
 		GStructToSourceLine.Add(Struct, MakeTuple(GetCurrentSourceFile()->AsShared(), Token.StartLine));
+		*/
+
 		return true;
 	}
 
@@ -5943,7 +5973,12 @@ bool FHeaderParser::CompileDeclaration(FClasses& AllClasses, TArray<UDelegateFun
 
 							UClass* CurrentClass = GetCurrentClass();
 
-							GClassSerializerMap.Add(CurrentClass, { ArchiveType, MoveTemp(EnclosingDefine) });
+							{
+								// UHTLite
+								FRWScopeLock Lock(ClassSerializerMapLock, SLT_Write);
+
+								ClassSerializerMap.Add(CurrentClass, { ArchiveType, MoveTemp(EnclosingDefine) });
+							}
 						}
 						else
 						{
@@ -6329,7 +6364,7 @@ UClass* FHeaderParser::CompileClassDeclaration(FClasses& AllClasses)
 	
 	FClass* Class = ParseClassNameDeclaration(AllClasses, /*out*/ DeclaredClassName, /*out*/ RequiredAPIMacroIfPresent);
 	check(Class);
-	FClassDeclarationMetaData* ClassDeclarationData = &GClassDeclarations.FindChecked(Class->GetFName());
+	FClassDeclarationMetaData* ClassDeclarationData = &ClassDeclarations.FindChecked(Class->GetFName());
 
 	ClassDefinitionRanges.Add(Class, ClassDefinitionRange(&Input[InputPos], nullptr));
 
@@ -6832,7 +6867,7 @@ void FHeaderParser::RedirectTypeIdentifier(FToken& Token) const
 {
 	check(Token.TokenType == TOKEN_Identifier);
 
-	FString* FoundRedirect = TypeRedirectMap.Find(Token.Identifier);
+	const FString* FoundRedirect = TypeRedirectMap.Find(Token.Identifier);
 	if (FoundRedirect)
 	{
 		Token.SetIdentifier(**FoundRedirect);
@@ -7003,6 +7038,7 @@ void FHeaderParser::ParseParameterList(FClasses& AllClasses, UFunction* Function
 	} while( MatchSymbol(TEXT(',')) );
 	RequireSymbol( TEXT(')'), TEXT("parameter list") );
 }
+
 UDelegateFunction* FHeaderParser::CompileDelegateDeclaration(FClasses& AllClasses, const TCHAR* DelegateIdentifier, EDelegateSpecifierAction::Type SpecifierAction)
 {
 	const TCHAR* CurrentScopeName = TEXT("Delegate Declaration");
@@ -7234,6 +7270,8 @@ UDelegateFunction* FHeaderParser::CompileDelegateDeclaration(FClasses& AllClasse
 	return DelegateSignatureFunction;
 }
 
+// UHTLite NOTE: Not required for code-gen. Will be refactored later.
+/*
 // Compares the properties of two functions to see if they have the same signature.
 bool AreFunctionSignaturesEqual(const UFunction* Lhs, const UFunction* Rhs)
 {
@@ -7307,6 +7345,7 @@ bool AreFunctionSignaturesEqual(const UFunction* Lhs, const UFunction* Rhs)
 		++RhsPropIter;
 	}
 }
+*/
 
 /**
  * Parses and compiles a function declaration
@@ -8316,6 +8355,8 @@ void FHeaderParser::ComputeFunctionParametersSize( UClass* Class )
  * @param	NestCount	number of nest levels to consume. if 0, consumes a single statement
  * @param	ErrorTag	text to use in error message if EOF is encountered before we've done
  */
+// UHTLite NOTE: Not required for code-gen. Will be refactored later.
+/*
 void FHeaderParser::SkipStatements( int32 NestCount, const TCHAR* ErrorTag  )
 {
 	FToken Token;
@@ -8350,6 +8391,7 @@ void FHeaderParser::SkipStatements( int32 NestCount, const TCHAR* ErrorTag  )
 		FError::Throwf(TEXT("Extraneous closing brace found in %s"), ErrorTag);
 	}
 }
+*/
 
 /*-----------------------------------------------------------------------------
 	Main script compiling routine.
@@ -8563,16 +8605,16 @@ ECompilationResult::Type FHeaderParser::ParseHeader(FClasses& AllClasses, FUnrea
 	Global functions.
 -----------------------------------------------------------------------------*/
 
-ECompilationResult::Type FHeaderParser::ParseRestOfModulesSourceFiles(FClasses& AllClasses, UPackage* ModulePackage, FHeaderParser& HeaderParser)
+ECompilationResult::Type FHeaderParser::ParseRestOfModulesSourceFiles(FClasses& AllClasses, UPackage* ModulePackage)
 {
-	if (const TArray<FUnrealSourceFile*>* SourceFiles = GUnrealSourceFilesMap.FindFilesForPackage(ModulePackage))
+	if (const TArray<FUnrealSourceFile*>* SourceFiles = UnrealSourceFilesMap.FindFilesForPackage(ModulePackage))
 	{
 		for (FUnrealSourceFile* SourceFile : *SourceFiles)
 		{
 			if (!SourceFile->IsParsed() || SourceFile->GetDefinedClassesCount() == 0)
 			{
 				ECompilationResult::Type Result;
-				if ((Result = ParseHeaders(AllClasses, HeaderParser, SourceFile)) != ECompilationResult::Succeeded)
+				if ((Result = ParseHeaders(AllClasses, SourceFile)) != ECompilationResult::Succeeded)
 				{
 					return Result;
 				}
@@ -8586,7 +8628,7 @@ ECompilationResult::Type FHeaderParser::ParseRestOfModulesSourceFiles(FClasses&
 // Parse Class's annotated headers and optionally its child classes.
 static const FString ObjectHeader(TEXT("NoExportTypes.h"));
 
-ECompilationResult::Type FHeaderParser::ParseHeaders(FClasses& AllClasses, FHeaderParser& HeaderParser, FUnrealSourceFile* SourceFile)
+ECompilationResult::Type FHeaderParser::ParseHeaders(FClasses& AllClasses, FUnrealSourceFile* SourceFile)
 {
 	ECompilationResult::Type Result = ECompilationResult::Succeeded;
 
@@ -8606,7 +8648,7 @@ ECompilationResult::Type FHeaderParser::ParseHeaders(FClasses& AllClasses, FHead
 			continue;
 		}
 
-		if (FUnrealSourceFile* DepFile = Include.Resolve())
+		if (FUnrealSourceFile* DepFile = Include.Resolve(UnrealSourceFilesMap, TypeDefinitionInfoMap))
 		{
 			SourceFilesRequired.Add(DepFile);
 		}
@@ -8617,7 +8659,7 @@ ECompilationResult::Type FHeaderParser::ParseHeaders(FClasses& AllClasses, FHead
 		UClass* Class = ClassDataPair.Key;
 		for (UClass* ParentClass = Class->GetSuperClass(); ParentClass && !ParentClass->HasAnyClassFlags(CLASS_Parsed | CLASS_Intrinsic); ParentClass = ParentClass->GetSuperClass())
 		{
-			SourceFilesRequired.Add(&GTypeDefinitionInfoMap[ParentClass]->GetUnrealSourceFile());
+			SourceFilesRequired.Add(&TypeDefinitionInfoMap[ParentClass]->GetUnrealSourceFile());
 		}
 	}
 
@@ -8625,7 +8667,7 @@ ECompilationResult::Type FHeaderParser::ParseHeaders(FClasses& AllClasses, FHead
 	{
 		SourceFile->GetScope()->IncludeScope(&RequiredFile->GetScope().Get());
 
-		ECompilationResult::Type ParseResult = ParseHeaders(AllClasses, HeaderParser, RequiredFile);
+		ECompilationResult::Type ParseResult = ParseHeaders(AllClasses, RequiredFile);
 
 		if (ParseResult != ECompilationResult::Succeeded)
 		{
@@ -8635,7 +8677,7 @@ ECompilationResult::Type FHeaderParser::ParseHeaders(FClasses& AllClasses, FHead
 
 	// Parse the file
 	{
-		ECompilationResult::Type OneFileResult = HeaderParser.ParseHeader(AllClasses, SourceFile);
+		ECompilationResult::Type OneFileResult = ParseHeader(AllClasses, SourceFile);
 
 		for (const TPair<UClass*, FSimplifiedParsingClassInfo>& ClassDataPair : SourceFile->GetDefinedClassesWithParsingInfo())
 		{
@@ -8654,6 +8696,8 @@ ECompilationResult::Type FHeaderParser::ParseHeaders(FClasses& AllClasses, FHead
 	return Result;
 }
 
+// UHTLite NOTE: Not required for code-gen. Will be refactored later.
+/*
 bool FHeaderParser::DependentClassNameFromHeader(const TCHAR* HeaderFilename, FString& OutClassName)
 {
 	FString DependentClassName(HeaderFilename);
@@ -8666,6 +8710,7 @@ bool FHeaderParser::DependentClassNameFromHeader(const TCHAR* HeaderFilename, FS
 	}
 	return false;
 }
+*/
 
 /**
  * Gets source files ordered by UCLASSes inheritance.
@@ -8675,7 +8720,7 @@ bool FHeaderParser::DependentClassNameFromHeader(const TCHAR* HeaderFilename, FS
  *
  * @returns Array of source files.
  */
-TSet<FUnrealSourceFile*> GetSourceFilesWithInheritanceOrdering(UPackage* CurrentPackage, FClasses& AllClasses)
+TSet<FUnrealSourceFile*> FHeaderParser::GetSourceFilesWithInheritanceOrdering(UPackage* CurrentPackage, FClasses& AllClasses) const
 {
 	TSet<FUnrealSourceFile*> SourceFiles;
 
@@ -8684,7 +8729,7 @@ TSet<FUnrealSourceFile*> GetSourceFilesWithInheritanceOrdering(UPackage* Current
 	// First add source files with the inheritance order.
 	for (UClass* Class : Classes)
 	{
-		TSharedRef<FUnrealTypeDefinitionInfo>* DefinitionInfoPtr = GTypeDefinitionInfoMap.Find(Class);
+		const TSharedRef<FUnrealTypeDefinitionInfo>* DefinitionInfoPtr = TypeDefinitionInfoMap.Find(Class);
 		if (DefinitionInfoPtr == nullptr)
 		{
 			continue;
@@ -8699,7 +8744,7 @@ TSet<FUnrealSourceFile*> GetSourceFilesWithInheritanceOrdering(UPackage* Current
 	}
 
 	// Then add the rest.
-	if (const TArray<FUnrealSourceFile*>* SourceFilesForPackage = GUnrealSourceFilesMap.FindFilesForPackage(CurrentPackage))
+	if (const TArray<FUnrealSourceFile*>* SourceFilesForPackage = UnrealSourceFilesMap.FindFilesForPackage(CurrentPackage))
 	{
 		for (FUnrealSourceFile* SourceFile : *SourceFilesForPackage)
 		{
@@ -8718,8 +8763,10 @@ void FHeaderParser::ExportNativeHeaders(
 	UPackage* CurrentPackage,
 	FClasses& AllClasses,
 	bool bAllowSaveExportedHeaders,
-	const FManifestModule& Module
-)
+	const FManifestModule& Module,
+	const FPublicSourceFileSet& PublicSourceFileSet,
+	const TMap<UPackage*, const FManifestModule*>& PackageToManifestModuleMap
+) const
 {
 	TSet<FUnrealSourceFile*> SourceFiles = GetSourceFilesWithInheritanceOrdering(CurrentPackage, AllClasses);
 	if (SourceFiles.Num() > 0)
@@ -8735,21 +8782,40 @@ void FHeaderParser::ExportNativeHeaders(
 
 		// Export native class definitions to package header files.
 		FNativeClassHeaderGenerator(
+			*this,
 			CurrentPackage,
 			SourceFiles,
 			AllClasses,
-			bAllowSaveExportedHeaders
+			bAllowSaveExportedHeaders,
+			PublicSourceFileSet,
+			PackageToManifestModuleMap
 		);
 	}
 }
 
-FHeaderParser::FHeaderParser(FFeedbackContext* InWarn, const FManifestModule& InModule)
+FHeaderParser::FHeaderParser(FFeedbackContext* InWarn,
+	const FManifestModule& InModule,
+	const FUnrealSourceFiles& InUnrealSourceFilesMap,
+	FTypeDefinitionInfoMap& InTypeDefinitionInfoMap,
+	const FClassDeclarations& InClassDeclarations,
+	TMap<UEnum*, EUnderlyingEnumType>& InEnumUnderlyingTypes,
+	FRWLock& InEnumUnderlyingTypesLock,
+	TMap<UClass*, FArchiveTypeDefinePair>& InClassSerializerMap,
+	FRWLock& InClassSerializerMapLock)
 	: FBaseParser()
+	, DefaultGeneratedCodeVersion(EGeneratedCodeVersion::V1)
+	, UnrealSourceFilesMap(InUnrealSourceFilesMap)
+	, TypeDefinitionInfoMap(InTypeDefinitionInfoMap)
+	, ClassDeclarations(InClassDeclarations)
 	, Warn(InWarn)
 	, bSpottedAutogeneratedHeaderInclude(false)
 	, NestLevel(0)
 	, TopNest(nullptr)
 	, CurrentlyParsedModule(&InModule)
+	, EnumUnderlyingTypes(InEnumUnderlyingTypes)
+	, EnumUnderlyingTypesLock(InEnumUnderlyingTypesLock)
+	, ClassSerializerMap(InClassSerializerMap)
+	, ClassSerializerMapLock(InClassSerializerMapLock)
 {
 	// Determine if the current module is part of the engine or a game (we are more strict about things for Engine modules)
 	switch (InModule.ModuleType)
@@ -8782,7 +8848,8 @@ FHeaderParser::FHeaderParser(FFeedbackContext* InWarn, const FManifestModule& In
 
 	FScriptLocation::Compiler = this;
 
-	static bool bConfigOptionsInitialized = false;
+	// UHTLite
+	/*static*/ bool bConfigOptionsInitialized = false;
 
 	if (!bConfigOptionsInitialized)
 	{
@@ -8852,10 +8919,11 @@ FString FHeaderParser::RequireExactlyOneSpecifierValue(const FPropertySpecifier&
 	return Specifier.Values[0];
 }
 
-// Exports the class to all vailable plugins
-void ExportClassToScriptPlugins(UClass* Class, const FManifestModule& Module, IScriptGeneratorPluginInterface& ScriptPlugin)
+#pragma region Plugins
+// Exports the class to all available plugins
+void ExportClassToScriptPlugins(UClass* Class, const FManifestModule& Module, IScriptGeneratorPluginInterface& ScriptPlugin, const FTypeDefinitionInfoMap& TypeDefinitionInfoMap)
 {
-	TSharedRef<FUnrealTypeDefinitionInfo>* DefinitionInfoRef = GTypeDefinitionInfoMap.Find(Class);
+	const TSharedRef<FUnrealTypeDefinitionInfo>* DefinitionInfoRef = TypeDefinitionInfoMap.Find(Class);
 	if (DefinitionInfoRef == nullptr)
 	{
 		const FString Empty = TEXT("");
@@ -8869,20 +8937,21 @@ void ExportClassToScriptPlugins(UClass* Class, const FManifestModule& Module, IS
 }
 
 // Exports class tree to all available plugins
-void ExportClassTreeToScriptPlugins(const FClassTree* Node, const FManifestModule& Module, IScriptGeneratorPluginInterface& ScriptPlugin)
+void ExportClassTreeToScriptPlugins(const FClassTree* Node, const FManifestModule& Module, IScriptGeneratorPluginInterface& ScriptPlugin, const FTypeDefinitionInfoMap& TypeDefinitionInfoMap)
 {
 	for (int32 ChildIndex = 0; ChildIndex < Node->NumChildren(); ++ChildIndex)
 	{
 		const FClassTree* ChildNode = Node->GetChild(ChildIndex);
-		ExportClassToScriptPlugins(ChildNode->GetClass(), Module, ScriptPlugin);
+		ExportClassToScriptPlugins(ChildNode->GetClass(), Module, ScriptPlugin, TypeDefinitionInfoMap);
 	}
 
 	for (int32 ChildIndex = 0; ChildIndex < Node->NumChildren(); ++ChildIndex)
 	{
 		const FClassTree* ChildNode = Node->GetChild(ChildIndex);
-		ExportClassTreeToScriptPlugins(ChildNode, Module, ScriptPlugin);
+		ExportClassTreeToScriptPlugins(ChildNode, Module, ScriptPlugin, TypeDefinitionInfoMap);
 	}
 }
+#pragma endregion Plugins
 
 // Parse all headers for classes that are inside CurrentPackage.
 ECompilationResult::Type FHeaderParser::ParseAllHeadersInside(
@@ -8890,7 +8959,18 @@ ECompilationResult::Type FHeaderParser::ParseAllHeadersInside(
 	FFeedbackContext* Warn,
 	UPackage* CurrentPackage,
 	const FManifestModule& Module,
-	TArray<IScriptGeneratorPluginInterface*>& ScriptPlugins
+#pragma region Plugins
+	TArray<IScriptGeneratorPluginInterface*>& ScriptPlugins,
+#pragma endregion Plugins
+	const FUnrealSourceFiles& UnrealSourceFilesMap,
+	FTypeDefinitionInfoMap& TypeDefinitionInfoMap,
+	const FPublicSourceFileSet& PublicSourceFileSet,
+	const TMap<UPackage*, const FManifestModule*>& PackageToManifestModuleMap,
+	const FClassDeclarations& ClassDeclarations,
+	TMap<UEnum*, EUnderlyingEnumType>& EnumUnderlyingTypes,
+	FRWLock& EnumUnderlyingTypesLock,
+	TMap<UClass*, FArchiveTypeDefinePair>& ClassSerializerMap,
+	FRWLock& ClassSerializerMapLock
 	)
 {
 	SCOPE_SECONDS_COUNTER_UHT(ParseAllHeaders);
@@ -8899,7 +8979,15 @@ ECompilationResult::Type FHeaderParser::ParseAllHeadersInside(
 	TGuardValue<bool> AutoRestoreVerifyObjectRefsFlag(GVerifyObjectReferencesOnly, true);
 	// Create the header parser and register it as the warning context.
 	// Note: This must be declared outside the try block, since the catch block will log into it.
-	FHeaderParser HeaderParser(Warn, Module);
+	FHeaderParser HeaderParser(Warn,
+		Module,
+		UnrealSourceFilesMap,
+		TypeDefinitionInfoMap,
+		ClassDeclarations,
+		EnumUnderlyingTypes,
+		EnumUnderlyingTypesLock,
+		ClassSerializerMap,
+		ClassSerializerMapLock);
 	Warn->SetContext(&HeaderParser);
 
 
@@ -8914,19 +9002,19 @@ ECompilationResult::Type FHeaderParser::ParseAllHeadersInside(
 
 		// Set up a filename for the error context if we don't even get as far parsing a class
 		FClass*                                      RootClass          = ModuleClasses.GetRootClass();
-		const TSharedRef<FUnrealTypeDefinitionInfo>& TypeDefinitionInfo = GTypeDefinitionInfoMap[RootClass];
+		const TSharedRef<FUnrealTypeDefinitionInfo>& TypeDefinitionInfo = TypeDefinitionInfoMap[RootClass];
 		const FUnrealSourceFile&                     RootSourceFile     = TypeDefinitionInfo->GetUnrealSourceFile();
 		const FString&                               RootFilename       = RootSourceFile.GetFilename();
 
 		HeaderParser.Filename = IFileManager::Get().ConvertToAbsolutePathForExternalAppForRead(*RootFilename);
 
-		if (const TArray<FUnrealSourceFile*>* SourceFiles = GPublicSourceFileSet.FindFilesForPackage(CurrentPackage))
+		if (const TArray<FUnrealSourceFile*>* SourceFiles = PublicSourceFileSet.FindFilesForPackage(CurrentPackage))
 		{
 			for (FUnrealSourceFile* SourceFile : *SourceFiles)
 			{
 				if ((!SourceFile->IsParsed() || SourceFile->GetDefinedClassesCount() == 0))
 				{
-					Result = ParseHeaders(ModuleClasses, HeaderParser, SourceFile);
+					Result = HeaderParser.ParseHeaders(ModuleClasses, SourceFile);
 					if (Result != ECompilationResult::Succeeded)
 					{
 						return Result;
@@ -8936,7 +9024,7 @@ ECompilationResult::Type FHeaderParser::ParseAllHeadersInside(
 		}
 		if (Result == ECompilationResult::Succeeded)
 		{
-			Result = FHeaderParser::ParseRestOfModulesSourceFiles(ModuleClasses, CurrentPackage, HeaderParser);
+			Result = HeaderParser.ParseRestOfModulesSourceFiles(ModuleClasses, CurrentPackage);
 		}
 
 		if (Result == ECompilationResult::Succeeded)
@@ -8957,11 +9045,13 @@ ECompilationResult::Type FHeaderParser::ParseAllHeadersInside(
 			double ExportTime = 0.0;
 			{
 				FScopedDurationTimer Timer(ExportTime);
-				ExportNativeHeaders(
+				HeaderParser.ExportNativeHeaders(
 					CurrentPackage,
 					ModuleClasses,
 					Module.SaveExportedHeaders,
-					Module
+					Module,
+					PublicSourceFileSet,
+					PackageToManifestModuleMap
 				);
 			}
 			GHeaderCodeGenTime += ExportTime;
@@ -8987,6 +9077,7 @@ ECompilationResult::Type FHeaderParser::ParseAllHeadersInside(
 	// Unregister the header parser from the feedback context
 	Warn->SetContext(NULL);
 
+#pragma region Plugins
 	if (Result == ECompilationResult::Succeeded && ScriptPlugins.Num())
 	{
 		FScopedDurationTimer PluginTimeTracker(GPluginOverheadTime);
@@ -8996,11 +9087,12 @@ ECompilationResult::Type FHeaderParser::ParseAllHeadersInside(
 		{
 			if (Plugin->ShouldExportClassesForModule(Module.Name, Module.ModuleType, Module.GeneratedIncludeDirectory))
 			{
-				ExportClassToScriptPlugins(RootNode->GetClass(), Module, *Plugin);
-				ExportClassTreeToScriptPlugins(RootNode, Module, *Plugin);
+				ExportClassToScriptPlugins(RootNode->GetClass(), Module, *Plugin, TypeDefinitionInfoMap);
+				ExportClassTreeToScriptPlugins(RootNode, Module, *Plugin, TypeDefinitionInfoMap);
 			}
 		}
 	}
+#pragma endregion Plugins
 
 	return Result;
 }
@@ -9034,6 +9126,8 @@ bool FHeaderParser::ClassNameHasValidPrefix(const FString& InNameToCheck, const
 	return bNamesMatch;
 }
 
+// UHTLite NOTE: Not required for code-gen. Will be refactored later.
+/*
 void FHeaderParser::ParseClassName(const TCHAR* Temp, FString& ClassName)
 {
 	// Skip leading whitespace
@@ -9062,6 +9156,7 @@ void FHeaderParser::ParseClassName(const TCHAR* Temp, FString& ClassName)
 		ParseClassName(Temp, ClassName);
 	}
 }
+*/
 
 enum class EBlockDirectiveType
 {
@@ -9191,7 +9286,7 @@ const TCHAR* GetBlockDirectiveTypeString(EBlockDirectiveType DirectiveType)
 }
 
 // Performs a preliminary parse of the text in the specified buffer, pulling out useful information for the header generation process
-void FHeaderParser::SimplifiedClassParse(const TCHAR* Filename, const TCHAR* InBuffer, TArray<FSimplifiedParsingClassInfo>& OutParsedClassArray, TArray<FHeaderProvider>& DependentOn, FStringOutputDevice& ClassHeaderTextStrippedOfCppText)
+void FHeaderParser::SimplifiedClassParse(const TCHAR* Filename, const TCHAR* InBuffer, TArray<FSimplifiedParsingClassInfo>& OutParsedClassArray, TArray<FHeaderProvider>& DependentOn, FStringOutputDevice& ClassHeaderTextStrippedOfCppText, FClassDeclarations& ClassDeclarations)
 {
 	FHeaderPreParser Parser;
 	FString StrLine;
@@ -9559,11 +9654,11 @@ void FHeaderParser::SimplifiedClassParse(const TCHAR* Filename, const TCHAR* InB
 					}
 
 					FName StrippedInterfaceName;
-					Parser.ParseClassDeclaration(Filename, StartOfLine + (UInterfaceMacroDecl - Str), CurrentLine, TEXT("UINTERFACE"), /*out*/ StrippedInterfaceName, /*out*/ ClassName, /*out*/ BaseClassName, /*out*/ DependentOn, OutParsedClassArray);
+					Parser.ParseClassDeclaration(Filename, StartOfLine + (UInterfaceMacroDecl - Str), CurrentLine, TEXT("UINTERFACE"), /*out*/ StrippedInterfaceName, /*out*/ ClassName, /*out*/ BaseClassName, /*out*/ DependentOn, OutParsedClassArray, ClassDeclarations);
 					OutParsedClassArray.Add(FSimplifiedParsingClassInfo(MoveTemp(ClassName), MoveTemp(BaseClassName), CurrentLine, true));
 					if (!bFoundExportedClasses)
 					{
-						if (FClassDeclarationMetaData* Found = GClassDeclarations.Find(StrippedInterfaceName))
+						if (FClassDeclarationMetaData* Found = ClassDeclarations.Find(StrippedInterfaceName))
 						{
 							bFoundExportedClasses = !(Found->ClassFlags & CLASS_NoExport);
 						}
@@ -9581,11 +9676,11 @@ void FHeaderParser::SimplifiedClassParse(const TCHAR* Filename, const TCHAR* InB
 					}
 
 					FName StrippedClassName;
-					Parser.ParseClassDeclaration(Filename, StartOfLine + (UClassMacroDecl - Str), CurrentLine, TEXT("UCLASS"), /*out*/ StrippedClassName, /*out*/ ClassName, /*out*/ BaseClassName, /*out*/ DependentOn, OutParsedClassArray);
+					Parser.ParseClassDeclaration(Filename, StartOfLine + (UClassMacroDecl - Str), CurrentLine, TEXT("UCLASS"), /*out*/ StrippedClassName, /*out*/ ClassName, /*out*/ BaseClassName, /*out*/ DependentOn, OutParsedClassArray, ClassDeclarations);
 					OutParsedClassArray.Add(FSimplifiedParsingClassInfo(MoveTemp(ClassName), MoveTemp(BaseClassName), CurrentLine, false));
 					if (!bFoundExportedClasses)
 					{
-						if (FClassDeclarationMetaData* Found = GClassDeclarations.Find(StrippedClassName))
+						if (FClassDeclarationMetaData* Found = ClassDeclarations.Find(StrippedClassName))
 						{
 							bFoundExportedClasses = !(Found->ClassFlags & CLASS_NoExport);
 						}
@@ -9606,7 +9701,7 @@ void FHeaderParser::SimplifiedClassParse(const TCHAR* Filename, const TCHAR* InB
 /////////////////////////////////////////////////////
 // FHeaderPreParser
 
-void FHeaderPreParser::ParseClassDeclaration(const TCHAR* Filename, const TCHAR* InputText, int32 InLineNumber, const TCHAR* StartingMatchID, FName& out_StrippedClassName, FString& out_ClassName, FString& out_BaseClassName, TArray<FHeaderProvider>& out_RequiredIncludes, const TArray<FSimplifiedParsingClassInfo>& ParsedClassArray)
+void FHeaderPreParser::ParseClassDeclaration(const TCHAR* Filename, const TCHAR* InputText, int32 InLineNumber, const TCHAR* StartingMatchID, FName& out_StrippedClassName, FString& out_ClassName, FString& out_BaseClassName, TArray<FHeaderProvider>& out_RequiredIncludes, const TArray<FSimplifiedParsingClassInfo>& ParsedClassArray, FClassDeclarations& ClassDeclarations)
 {
 	const TCHAR* ErrorMsg = TEXT("Class declaration");
 
@@ -9646,7 +9741,7 @@ void FHeaderPreParser::ParseClassDeclaration(const TCHAR* Filename, const TCHAR*
 			return DeclarationData;
 		};
 
-		GClassDeclarations.AddIfMissing(out_StrippedClassName, MoveTemp(ConstructDeclarationData));
+		ClassDeclarations.AddIfMissing(out_StrippedClassName, MoveTemp(ConstructDeclarationData));
 	}
 
 	// Skip optional final keyword
diff --git a/Engine/Source/Programs/UHTLite/Private/HeaderParser.h b/Engine/Source/Programs/UHTLite/Private/HeaderParser.h
--- a/Engine/Source/Programs/UHTLite/Private/HeaderParser.h
+++ b/Engine/Source/Programs/UHTLite/Private/HeaderParser.h
@@ -390,8 +390,6 @@ struct ClassDefinitionRange
 	bool bHasGeneratedBody;
 };
 
-extern TMap<UClass*, ClassDefinitionRange> ClassDefinitionRanges;
-
 #ifndef UHT_DOCUMENTATION_POLICY_DEFAULT
 #define UHT_DOCUMENTATION_POLICY_DEFAULT false
 #endif
@@ -416,7 +414,7 @@ class FHeaderParser : public FBaseParser, public FContextSupplier
 {
 public:
 	// Default version of generated code. Defaults to oldest possible, unless specified otherwise in config.
-	static EGeneratedCodeVersion DefaultGeneratedCodeVersion;
+	EGeneratedCodeVersion DefaultGeneratedCodeVersion;
 
 	// Compute the function parameter size and save the return offset
 	static void ComputeFunctionParametersSize(UClass* InClass);
@@ -427,7 +425,18 @@ public:
 		FFeedbackContext* Warn,
 		UPackage* LimitOuter,
 		const FManifestModule& Module,
-		TArray<class IScriptGeneratorPluginInterface*>& ScriptPlugins
+#pragma region Plugins
+		TArray<class IScriptGeneratorPluginInterface*>& ScriptPlugins,
+#pragma endregion Plugins
+		const FUnrealSourceFiles& UnrealSourceFilesMap,
+		FTypeDefinitionInfoMap& TypeDefinitionInfoMap,
+		const FPublicSourceFileSet& PublicSourceFileSet,
+		const TMap<UPackage*, const FManifestModule*>& PackageToManifestModuleMap,
+		const FClassDeclarations& ClassDeclarations,
+		TMap<UEnum*, EUnderlyingEnumType>& EnumUnderlyingTypes,
+		FRWLock& EnumUnderlyingTypesLock,
+		TMap<UClass*, FArchiveTypeDefinePair>& InClassSerializerMap,
+		FRWLock& InClassSerializerMapLock
 	);
 
 	// Performs a preliminary parse of the text in the specified buffer, pulling out:
@@ -437,7 +446,7 @@ public:
 	//   
 	//  It also splits the buffer up into:
 	//   ScriptText (text outside of #if CPP and #if DEFAULTS blocks)
-	static void SimplifiedClassParse(const TCHAR* Filename, const TCHAR* Buffer, TArray<FSimplifiedParsingClassInfo>& OutParsedClassArray, TArray<FHeaderProvider>& DependentOn, FStringOutputDevice& ScriptText);
+	static void SimplifiedClassParse(const TCHAR* Filename, const TCHAR* Buffer, TArray<FSimplifiedParsingClassInfo>& OutParsedClassArray, TArray<FHeaderProvider>& DependentOn, FStringOutputDevice& ScriptText, FClassDeclarations& ClassDeclarations);
 
 	/**
 	 * Returns True if the given class name includes a valid Unreal prefix and matches up with the given original class Name.
@@ -454,7 +463,10 @@ public:
 	 * @param OutClass The resulting class name (if successfull)
 	 * @return true if the filename was a header filename (.h), false otherwise (in which case OutClassName is unmodified).
 	 */
+	// UHTLite NOTE: Not required for code-gen. Will be refactored later.
+	/*
 	static bool DependentClassNameFromHeader(const TCHAR* HeaderFilename, FString& OutClassName);
+	*/
 
 	/**
 	 * Transforms CPP-formated string containing default value, to inner formated string
@@ -476,12 +488,17 @@ public:
 	 *
 	 * @return	Result enumeration.
 	 */
-	static ECompilationResult::Type ParseHeaders(FClasses& AllClasses, FHeaderParser& HeaderParser, FUnrealSourceFile* SourceFile);
+	ECompilationResult::Type ParseHeaders(FClasses& AllClasses, FUnrealSourceFile* SourceFile);
 
 protected:
 	friend struct FScriptLocation;
 	friend struct FNativeClassHeaderGenerator;
 
+	// UHTLite
+	const FUnrealSourceFiles& UnrealSourceFilesMap;
+	FTypeDefinitionInfoMap&   TypeDefinitionInfoMap;
+	const FClassDeclarations& ClassDeclarations;
+
 	// For compiling messages and errors.
 	FFeedbackContext* Warn;
 
@@ -648,17 +665,17 @@ protected:
 	////////////////////////////////////////////////////
 
 	// Special parsed struct names that do not require a prefix
-	static TArray<FString> StructsWithNoPrefix;
+	TArray<FString> StructsWithNoPrefix;
 	
 	// Special parsed struct names that have a 'T' prefix
-	static TArray<FString> StructsWithTPrefix;
+	TArray<FString> StructsWithTPrefix;
 
 	// Mapping from 'human-readable' macro substring to # of parameters for delegate declarations
 	// Index 0 is 1 parameter, Index 1 is 2, etc...
-	static TArray<FString> DelegateParameterCountStrings;
+	TArray<FString> DelegateParameterCountStrings;
 
 	// Types that have been renamed, treat the old deprecated name as the new name for code generation
-	static TMap<FString, FString> TypeRedirectMap;
+	TMap<FString, FString> TypeRedirectMap;
 
 	// List of all used identifiers for net service function declarations (every function must be unique)
 	TMap<int32, FString> UsedRPCIds;
@@ -666,10 +683,19 @@ protected:
 	TMap<int32, FString> RPCsNeedingHookup;
 
 	// List of all multiplex methods defined on structs
-	static FRigVMStructMap StructRigVMMap;
+	FRigVMStructMap StructRigVMMap;
 
 	// Constructor.
-	explicit FHeaderParser(FFeedbackContext* InWarn, const FManifestModule& InModule);
+	explicit FHeaderParser(FFeedbackContext* InWarn,
+		const FManifestModule& InModule,
+		const FUnrealSourceFiles& InUnrealSourceFilesMap,
+		FTypeDefinitionInfoMap& InTypeDefinitionInfoMap,
+		const FClassDeclarations& InClassDeclarations,
+		TMap<UEnum*, EUnderlyingEnumType>& InEnumUnderlyingTypes,
+		FRWLock& InEnumUnderlyingTypesLock,
+		TMap<UClass*, FArchiveTypeDefinePair>& InClassSerializerMap,
+		FRWLock& InClassSerializerMapLock
+		);
 
 	virtual ~FHeaderParser()
 	{
@@ -725,10 +751,13 @@ protected:
 	 *
 	 * @return	Result enumeration.
 	 */
-	static ECompilationResult::Type ParseRestOfModulesSourceFiles(FClasses& AllClasses, UPackage* ModulePackage, FHeaderParser& HeaderParser);
+	ECompilationResult::Type ParseRestOfModulesSourceFiles(FClasses& AllClasses, UPackage* ModulePackage);
 
 	//@TODO: Remove this method
+	// UHTLite NOTE: Not required for code-gen. Will be refactored later.
+	/*
 	static void ParseClassName(const TCHAR* Temp, FString& ClassName);
+	*/
 
 	/**
 	 * @param		Input		An input string, expected to be a script comment.
@@ -750,12 +779,30 @@ protected:
 	 * @param AllClasses The class tree for CurrentPackage.
 	 * @param Module Currently exported module.
 	 */
-	static void ExportNativeHeaders(
+	void ExportNativeHeaders(
 		UPackage* CurrentPackage,
 		FClasses& AllClasses,
 		bool bAllowSaveExportedHeaders,
-		const FManifestModule& Module
-	);
+		const FManifestModule& Module,
+		const FPublicSourceFileSet& PublicSourceFileSet,
+		const TMap<UPackage*, const FManifestModule*>& PackageToManifestModuleMap
+	) const;
+
+	// UHTLite
+	bool DoesAnythingInHierarchyHaveDefaultToInstanced(UClass* TestClass) const;
+
+	// UHTLite
+	TSet<FUnrealSourceFile*> GetSourceFilesWithInheritanceOrdering(UPackage* CurrentPackage, FClasses& AllClasses) const;
+
+	// UHTLite
+	FProperty* CreateVariableProperty(FPropertyBase& VarProperty, FFieldVariant Scope, const FName& Name, EObjectFlags ObjectFlags, EVariableCategory::Type VariableCategory, FUnrealSourceFile* UnrealSourceFile) const;
+
+	// UHTLite
+	void AddModuleRelativePathToMetadata(FUnrealSourceFile& SourceFile, TMap<FName, FString>& MetaData) const;
+	void AddModuleRelativePathToMetadata(UField* Type, TMap<FName, FString>& MetaData) const;
+
+	// UHTLite
+	void AddIncludePathToMetadata(UField* Type, TMap<FName, FString>& MetaData) const;
 
 	// FContextSupplier interface.
 	virtual FString GetContext() override;
@@ -831,9 +878,12 @@ protected:
 	// nesting level.
 	void CheckAllow(const TCHAR* Thing, ENestAllowFlags AllowFlags);
 
-	UStruct* GetSuperScope( UStruct* CurrentScope, const FName& SearchName );	
+	// UHTLite NOTE: Not required for code-gen. Will be refactored later.
+	/*
+	UStruct* GetSuperScope( UStruct* CurrentScope, const FName& SearchName );
 
 	void SkipStatements( int32 SubCount, const TCHAR* ErrorTag );
+	*/
 
 	/**
 	 * Parses a variable or return value declaration and determines the variable type and property flags.
@@ -881,12 +931,18 @@ protected:
 	 *
 	 * @return	true if the specified class is an intrinsic type or if the class has successfully been parsed
 	 */
+	// UHTLite NOTE: Not required for code-gen. Will be refactored later.
+	/*
 	bool AllowReferenceToClass(UStruct* Scope, UClass* CheckClass) const;
+	*/
 
 	/**
 	 * @return	true if Scope has FProperty objects in its list of fields
 	 */
+	// UHTLite NOTE: Not required for code-gen. Will be refactored later.
+	/*
 	static bool HasMemberProperties( const UStruct* Scope );
+	*/
 
 	/**
 	 * Parses optional metadata text.
@@ -1026,7 +1082,7 @@ private:
 	static FDocumentationPolicy GetDocumentationPolicyForStruct(UStruct* Struct);
 
 	// Property types to provide UI Min and Max ranges
-	static TArray<FString> PropertyCPPTypesRequiringUIRanges;
+	static const TArray<FString> PropertyCPPTypesRequiringUIRanges;
 
 	// Returns true if a given CPP types required ui checking
 	static bool DoesCPPTypeRequireDocumentation(const FString& CPPType);
@@ -1044,7 +1100,7 @@ private:
 	bool CheckUIMinMaxRangeFromMetaData(const FString& UIMin, const FString& UIMax);
 
 	// Names that cannot be used enums, UStructs, or UClasses
-	static TArray<FString> ReservedTypeNames;
+	static const TArray<FString> ReservedTypeNames;
 
 public:
 	/**
@@ -1062,6 +1118,21 @@ public:
 	* @return True if the Token is using a reserved name
 	*/
 	static bool IsReservedTypeName(const FToken& Token);
+
+	// UHTLite
+	TMap<UClass*, ClassDefinitionRange> ClassDefinitionRanges;
+	TMap<FProperty*, FString>           ArrayDimensions;
+	mutable TSet<FProperty*>            UnsizedProperties;
+	TSet<UField*>                       EditorOnlyDataTypes;
+	TSet<FProperty*>                    PropertyUsesMemoryImageAllocator;
+
+	// UHTLite NOTE: These are read-only in the code-generation phase
+	TMap<UEnum*, EUnderlyingEnumType>& EnumUnderlyingTypes;
+	FRWLock&                           EnumUnderlyingTypesLock;
+
+	// UHTLite NOTE: These are read-only in the code-generation phase
+	TMap<UClass*, FArchiveTypeDefinePair>& ClassSerializerMap;
+	FRWLock&                               ClassSerializerMapLock;
 };
 
 /////////////////////////////////////////////////////
@@ -1084,6 +1155,7 @@ public:
 		FString& out_ClassName,
 		FString& out_BaseClassName,
 		TArray<FHeaderProvider>& out_ClassNames,
-		const TArray<FSimplifiedParsingClassInfo>& ParsedClassArray
+		const TArray<FSimplifiedParsingClassInfo>& ParsedClassArray,
+		FClassDeclarations& ClassDeclarations
 	);
 };
diff --git a/Engine/Source/Programs/UHTLite/Private/HeaderProvider.cpp b/Engine/Source/Programs/UHTLite/Private/HeaderProvider.cpp
--- a/Engine/Source/Programs/UHTLite/Private/HeaderProvider.cpp
+++ b/Engine/Source/Programs/UHTLite/Private/HeaderProvider.cpp
@@ -13,19 +13,19 @@ FHeaderProvider::FHeaderProvider(EHeaderProviderSourceType InType, FString&& InI
 
 }
 
-FUnrealSourceFile* FHeaderProvider::Resolve()
+FUnrealSourceFile* FHeaderProvider::Resolve(const FUnrealSourceFiles& UnrealSourceFilesMap, const FTypeDefinitionInfoMap& TypeDefinitionInfoMap)
 {
 	if (Type != EHeaderProviderSourceType::Resolved)
 	{
 		if (Type == EHeaderProviderSourceType::ClassName)
 		{
 			FName IdName(*Id, FNAME_Find);
-			if (TSharedRef<FUnrealTypeDefinitionInfo>* Source = GTypeDefinitionInfoMap.FindByName(IdName))
+			if (const TSharedRef<FUnrealTypeDefinitionInfo>* Source = TypeDefinitionInfoMap.FindByName(IdName))
 			{
 				Cache = &(*Source)->GetUnrealSourceFile();
 			}
 		}
-		else if (const TSharedRef<FUnrealSourceFile>* Source = GUnrealSourceFilesMap.Find(Id))
+		else if (const TSharedRef<FUnrealSourceFile>* Source = UnrealSourceFilesMap.Find(Id))
 		{
 			Cache = &Source->Get();
 		}
diff --git a/Engine/Source/Programs/UHTLite/Private/HeaderProvider.h b/Engine/Source/Programs/UHTLite/Private/HeaderProvider.h
--- a/Engine/Source/Programs/UHTLite/Private/HeaderProvider.h
+++ b/Engine/Source/Programs/UHTLite/Private/HeaderProvider.h
@@ -20,7 +20,7 @@ class FHeaderProvider
 public:
 	FHeaderProvider(EHeaderProviderSourceType Type, FString&& Id);
 
-	FUnrealSourceFile* Resolve();
+	FUnrealSourceFile* Resolve(const struct FUnrealSourceFiles& UnrealSourceFilesMap, const struct FTypeDefinitionInfoMap& TypeDefinitionInfoMap);
 
 	FString ToString() const;
 
diff --git a/Engine/Source/Programs/UHTLite/Private/NativeClassExporter.h b/Engine/Source/Programs/UHTLite/Private/NativeClassExporter.h
--- a/Engine/Source/Programs/UHTLite/Private/NativeClassExporter.h
+++ b/Engine/Source/Programs/UHTLite/Private/NativeClassExporter.h
@@ -5,6 +5,7 @@
 #include "CoreMinimal.h"
 #include "Classes.h"
 #include "Async/TaskGraphInterfaces.h"
+#include "HeaderParser.h"
 
 class FUnrealSourceFile;
 class UPackage;
@@ -82,6 +83,8 @@ FString CreateUTF8LiteralString(const FString& Str);
 struct FNativeClassHeaderGenerator
 {
 private:
+	// UHTLite
+	const FHeaderParser& HeaderParser;
 	FString API;
 	FString APIStringPrivate;
 
@@ -133,7 +136,7 @@ private:
 	 * @param	Struct			UStruct to export properties
 	 * @param	TextIndent		Current text indentation
 	 */
-	static void ExportProperties(FOutputDevice& Out, UStruct* Struct, int32 TextIndent);
+	void ExportProperties(FOutputDevice& Out, UStruct* Struct, int32 TextIndent) const;
 
 	/** Return the name of the singleton function that returns the UObject for Item */
 	static const FString& GetPackageSingletonName(const UPackage* Item, TSet<FString>* UniqueCrossModuleReferences);
@@ -257,9 +260,9 @@ private:
 	 * @param	Structs		the struct to export
 	 * @param	TextIndent	the current indentation of the header exporter
 	 */
-	static void ExportMirrorsForNoexportStruct(FOutputDevice& Out, UScriptStruct* Struct, int32 TextIndent);
+	void ExportMirrorsForNoexportStruct(FOutputDevice& Out, UScriptStruct* Struct, int32 TextIndent) const;
 
-	/**heade
+	/**
 	 * Exports the parameter struct declarations for the list of functions specified
 	 * 
 	 * @param	Function	the function that (may) have parameters which need to be exported
@@ -293,7 +296,7 @@ private:
 	 * @param	Indent				number of spaces to put before each line
 	 * @param	bOutputConstructor	If true, output a constructor for the param struct
 	 */
-	static void ExportEventParm(FUHTStringBuilder& Out, TSet<FString>& PropertyFwd, UFunction* Function, int32 Indent, bool bOutputConstructor, EExportingState ExportingState);
+	void ExportEventParm(FUHTStringBuilder& Out, TSet<FString>& PropertyFwd, UFunction* Function, int32 Indent, bool bOutputConstructor, EExportingState ExportingState) const;
 
 	/**
 	* Move the temp header files into the .h files
@@ -328,7 +331,7 @@ private:
 	 * @param	FunctionHeaderStyle	Whether we're outputting a declaration or definition.
 	 * @param	ExtraParam			Optional extra parameter that will be added to the declaration as the first argument
 	 */
-	static void ExportNativeFunctionHeader(
+	void ExportNativeFunctionHeader(
 		FOutputDevice&                   Out,
 		TSet<FString>&                   OutFwdDecls,
 		const FFuncInfo&                 FunctionData,
@@ -336,7 +339,7 @@ private:
 		EExportFunctionHeaderStyle::Type FunctionHeaderStyle,
 		const TCHAR*                     ExtraParam,
 		const TCHAR*                     APIString
-	);
+	) const;
 
 	/**
 	* Runs checks whether necessary RPC functions exist for function described by FunctionData.
@@ -347,7 +350,10 @@ private:
 	* @param	ValidatePosition		Position in source file of _Validate function for function described by FunctionData.
 	* @param	SourceFile				Currently analyzed source file.
 	*/
+	// UHTLite NOTE: Not required for code-gen. Will be refactored later.
+	/*
 	void CheckRPCFunctions(FReferenceGatherers& OutReferenceGatherers, const FFuncInfo& FunctionData, const FString& ClassName, int32 ImplementationPosition, int32 ValidatePosition, const FUnrealSourceFile& SourceFile) const;
+	*/
 
 	/**
 	 * Exports the native stubs for the list of functions specified
@@ -397,7 +403,7 @@ private:
 	 * @param	Out			The destination to write to.
 	 * @param	Package		Package to export code for.
 	**/
-	void ExportGeneratedPackageInitCode(FOutputDevice& Out, const TCHAR* InDeclarations, const UPackage* Package, uint32 CRC);
+	void ExportGeneratedPackageInitCode(FOutputDevice& Out, const TCHAR* InDeclarations, const UPackage* Package, uint32 CRC, TMap<const UPackage*, TArray<UField*>>& PackageSingletons);
 
 	/**
 	 * Function to output the C++ code necessary to set up the given array of properties
@@ -444,7 +450,7 @@ private:
 	 * 
 	 * @param SourceFile	current source file
 	 */
-	static void ExportCallbackFunctions(
+	void ExportCallbackFunctions(
 		FOutputDevice&            OutGeneratedHeaderText,
 		FOutputDevice&            Out,
 		TSet<FString>&            OutFwdDecls,
@@ -452,7 +458,7 @@ private:
 		const TCHAR*              CallbackWrappersMacroName,
 		EExportCallbackType       ExportCallbackType,
 		const TCHAR*              APIString
-	);
+	) const;
 
 	/**
 	 * Determines if the property has alternate export text associated with it and if so replaces the text in PropertyText with the
@@ -499,6 +505,36 @@ private:
 	 */
 	static void ExportConstructorsMacros(FOutputDevice& OutGeneratedHeaderText, FOutputDevice& VTableOut, FOutputDevice& StandardUObjectConstructorsMacroCall, FOutputDevice& EnhancedUObjectConstructorsMacroCall, const FString& ConstructorsMacroPrefix, FClass* Class, const TCHAR* APIArg);
 
+	// UHTLite
+	TArray<FUnrealSourceFile*> GetSourceFilesInDependencyOrder(const TSet<FUnrealSourceFile*>& SourceFiles, const TSet<FUnrealSourceFile*>& Ignore) const;
+
+	// UHTLite
+	void GetSourceFilesInDependencyOrderRecursive(TArray<FUnrealSourceFile*>& OutTest, FUnrealSourceFile* SourceFile, TSet<const FUnrealSourceFile*>& VisitedSet, bool bCheckDependenciesOnly, const TSet<FUnrealSourceFile*>& Ignore) const;
+
+	// UHTLite
+	void WriteReplicatedMacroData(
+		const ClassDefinitionRange& ClassRange,
+		const TCHAR* ClassCPPName,
+		const TCHAR* APIArg,
+		FClass* Class,
+		FClass* SuperClass,
+		FOutputDevice& Writer,
+		const FUnrealSourceFile& SourceFile,
+		FNativeClassHeaderGenerator::EExportClassOutFlags& OutFlags) const;
+
+	// UHTLite
+	void ExportNetData(FOutputDevice& Out, UClass* Class, const TCHAR* APIArg) const;
+
+	// UHTLite NOTE: This function modifies GeneratedCodeHashes!
+	void AddGeneratedCodeHash(void* Field, uint32 Hash) const;
+
+	// UHTLite
+	FString GetGeneratedCodeHashTag(void* Field) const;
+
+	// UHTLite NOTE: Globals!
+	static TMap<void*, uint32> GeneratedCodeHashes;
+	static FRWLock             GeneratedCodeHashesLock;
+
 public:
 
 	/** 
@@ -535,10 +571,13 @@ public:
 
 	// Constructor
 	FNativeClassHeaderGenerator(
+		const FHeaderParser& InHeaderParser,
 		const UPackage* InPackage,
 		const TSet<FUnrealSourceFile*>& SourceFiles,
 		FClasses& AllClasses,
-		bool InAllowSaveExportedHeaders
+		bool InAllowSaveExportedHeaders,
+		const FPublicSourceFileSet& PublicSourceFileSet,
+		const TMap<UPackage*, const FManifestModule*>& PackageToManifestModuleMap
 	);
 
 	/**
@@ -547,7 +586,10 @@ public:
 	 * @param Function Function to get return type of.
 	 * @return FString with function return type.
 	 */
+	// UHTLite NOTE: Not required for code-gen. Will be refactored later.
+	/*
 	static FString GetFunctionReturnString(UFunction* Function, FReferenceGatherers& OutReferenceGatherers);
+	*/
 
 	/**
 	* Gets string with function parameters (with names).
@@ -555,7 +597,7 @@ public:
 	* @param Function Function to get parameters of.
 	* @return FString with function parameters.
 	*/
-	static FString GetFunctionParameterString(UFunction* Function, FReferenceGatherers& OutReferenceGatherers);
+	FString GetFunctionParameterString(UFunction* Function, FReferenceGatherers& OutReferenceGatherers) const;
 
 	/**
 	 * Checks if function is missing "virtual" specifier.
@@ -564,7 +606,10 @@ public:
 	 * @param FunctionNamePosition Position of name of function in SourceFile.
 	 * @return true if function misses "virtual" specifier, false otherwise.
 	 */
+	// UHTLite NOTE: Not required for code-gen. Will be refactored later.
+	/*
 	static bool IsMissingVirtualSpecifier(const FString& SourceFile, int32 FunctionNamePosition);
+	*/
 };
 
 ENUM_CLASS_FLAGS(FNativeClassHeaderGenerator::EExportClassOutFlags);
\ No newline at end of file
diff --git a/Engine/Source/Programs/UHTLite/Private/ParserHelper.cpp b/Engine/Source/Programs/UHTLite/Private/ParserHelper.cpp
--- a/Engine/Source/Programs/UHTLite/Private/ParserHelper.cpp
+++ b/Engine/Source/Programs/UHTLite/Private/ParserHelper.cpp
@@ -96,6 +96,8 @@ void FClassMetaData::AddInheritanceParent(UClass* ImplementedInterfaceClass, FUn
 /////////////////////////////////////////////////////
 // FPropertyBase
 
+// UHTLite NOTE: Not required for code-gen. Will be refactored later.
+/*
 const TCHAR* FPropertyBase::GetPropertyTypeText( EPropertyType Type )
 {
 	switch ( Type )
@@ -135,6 +137,7 @@ const TCHAR* FPropertyBase::GetPropertyTypeText( EPropertyType Type )
 
 	return TEXT("");
 }
+*/
 
 /////////////////////////////////////////////////////
 // FToken
@@ -246,12 +249,15 @@ FFunctionData* FFunctionData::FindForFunction(UFunction* Function)
 	return &(*Output).Get();
 }
 
+// UHTLite NOTE: Not required for code-gen. Will be refactored later.
+/*
 FFunctionData* FFunctionData::Add(UFunction* Function)
 {
 	TUniqueObj<FFunctionData>& Output = FunctionDataMap.Add(Function);
 
 	return &Output.Get();
 }
+*/
 
 FFunctionData* FFunctionData::Add(FFuncInfo&& FunctionInfo)
 {
diff --git a/Engine/Source/Programs/UHTLite/Private/ParserHelper.h b/Engine/Source/Programs/UHTLite/Private/ParserHelper.h
--- a/Engine/Source/Programs/UHTLite/Private/ParserHelper.h
+++ b/Engine/Source/Programs/UHTLite/Private/ParserHelper.h
@@ -680,6 +680,8 @@ public:
 		}
 	}
 
+	// UHTLite NOTE: Not required for code-gen. Will be refactored later.
+	/*
 	FString Describe()
 	{
 		return FString::Printf(
@@ -692,6 +694,8 @@ public:
 			MetaClass!=NULL?*MetaClass->GetName():TEXT("NULL")
 			);
 	}
+	*/
+
 	//@}
 
 	EIntType GetSizedIntTypeFromPropertyType(EPropertyType PropType)
@@ -713,7 +717,10 @@ public:
 		}
 	}
 
+	// UHTLite NOTE: Not required for code-gen. Will be refactored later.
+	/*
 	static const TCHAR* GetPropertyTypeText( EPropertyType Type );
+	*/
 
 	friend struct FPropertyBaseArchiveProxy;
 };
@@ -1017,6 +1024,8 @@ public:
 		else return 0;
 	}
 
+	// UHTLite NOTE: Not required for code-gen. Will be refactored later.
+	/*
 	FString Describe()
 	{
 		return FString::Printf(
@@ -1028,6 +1037,7 @@ public:
 			ImpliedPropertyFlags
 			);
 	}
+	*/
 
 	friend struct FTokenArchiveProxy;
 };
@@ -1235,6 +1245,8 @@ public:
 	 * 
 	 * @param	Indent	number of spaces to insert at the beginning of each line
 	 */	
+	// UHTLite NOTE: Not required for code-gen. Will be refactored later.
+	/*
 	void Dump( int32 Indent )
 	{
 		for (auto& Kvp : *this)
@@ -1247,6 +1259,7 @@ public:
 			}
 		}
 	}
+	*/
 };
 
 /**
@@ -1288,6 +1301,8 @@ public:
 	* 
 	* @param	Indent	number of spaces to insert at the beginning of each line
 	*/	
+	// UHTLite NOTE: Not required for code-gen. Will be refactored later.
+	/*
 	void Dump( int32 Indent )
 	{
 		UE_LOG(LogCompile, Log, TEXT("%s%s"), FCString::Spc(Indent), *StructData.Describe());
@@ -1295,6 +1310,7 @@ public:
 		UE_LOG(LogCompile, Log, TEXT("%sproperties:"), FCString::Spc(Indent));
 		StructPropertyData.Dump(Indent + 4);
 	}
+	*/
 
 	/** Constructor */
 	FStructData(FToken&& StructToken) : StructData(MoveTemp(StructToken)) {}
@@ -1392,6 +1408,8 @@ public:
 	 * 
 	 * @param	Indent	number of spaces to insert at the beginning of each line
 	 */	
+	// UHTLite NOTE: Not required for code-gen. Will be refactored later.
+	/*
 	void Dump( int32 Indent )
 	{
 		UE_LOG(LogCompile, Log, TEXT("%sparameters:"), FCString::Spc(Indent));
@@ -1403,6 +1421,7 @@ public:
 			UE_LOG(LogCompile, Log, TEXT("%s%s"), FCString::Spc(Indent + 4), *ReturnTypeData.Token.Describe());
 		}
 	}
+	*/
 
 	/**
 	 * Sets the specified function export flags
@@ -1428,7 +1447,10 @@ public:
 	/**
 	 * Adds function data object for given function object.
 	 */
+	// UHTLite NOTE: Not required for code-gen. Will be refactored later.
+	/*
 	static FFunctionData* Add(UFunction* Function);
+	*/
 
 	/**
 	 * Adds function data object for given function object.
@@ -1713,7 +1735,10 @@ public:
 	 * 
 	 * @param	Indent	number of spaces to insert at the beginning of each line
 	 */	
+	// UHTLite NOTE: Not required for code-gen. Will be refactored later.
+	/*
 	void Dump( int32 Indent );
+	*/
 
 	/**
 	 * Add a string to the list of inheritance parents for this class.
diff --git a/Engine/Source/Programs/UHTLite/Private/UnrealHeaderTool.h b/Engine/Source/Programs/UHTLite/Private/UnrealHeaderTool.h
--- a/Engine/Source/Programs/UHTLite/Private/UnrealHeaderTool.h
+++ b/Engine/Source/Programs/UHTLite/Private/UnrealHeaderTool.h
@@ -1,6 +1,12 @@
 // Copyright Epic Games, Inc. All Rights Reserved.
 #pragma once
+
 #include "CoreMinimal.h"
+#include "SimplifiedParsingClassInfo.h"
+#include "UnrealSourceFile.h"
+
+// UHTLite NOTE: Provided to ease debugging during refactoring.
+#define UHT_USE_PARALLEL_FOR 0
 
 struct FHeaderParserNames
 {
@@ -9,4 +15,11 @@ struct FHeaderParserNames
 	static const FName NAME_ShowCategories;
 	static const FName NAME_SparseClassDataTypes;
 	static const FName NAME_AdvancedClassDisplay;
-};
\ No newline at end of file
+};
+
+struct FPerHeaderData
+{
+	TSharedPtr<FUnrealSourceFile> UnrealSourceFile;
+	TArray<FHeaderProvider> DependsOn;
+	TArray<FSimplifiedParsingClassInfo> ParsedClassArray;
+};
diff --git a/Engine/Source/Programs/UHTLite/Private/UnrealHeaderToolMain.cpp b/Engine/Source/Programs/UHTLite/Private/UnrealHeaderToolMain.cpp
--- a/Engine/Source/Programs/UHTLite/Private/UnrealHeaderToolMain.cpp
+++ b/Engine/Source/Programs/UHTLite/Private/UnrealHeaderToolMain.cpp
@@ -14,6 +14,16 @@
 #include "Misc/CompilationResult.h"
 #include "UnrealHeaderToolGlobals.h"
 #include "RequiredProgramMainCPPInclude.h"
+#include "ClassMaps.h"
+#include "IScriptGeneratorPluginInterface.h"
+#include "ProfilingDebugging/ScopedTimers.h"
+#include "HeaderParser.h"
+#include "Features/IModularFeatures.h"
+#include "Manifest.h"
+#include "UnrealTypeDefinitionInfo.h"
+#include "StringUtils.h"
+#include "FileLineException.h"
+#include "NativeClassExporter.h"
 
 IMPLEMENT_APPLICATION(UnrealHeaderTool, "UnrealHeaderTool");
 
@@ -22,6 +32,814 @@ DEFINE_LOG_CATEGORY(LogCompile);
 bool GUHTWarningLogged = false;
 bool GUHTErrorLogged = false;
 
+// UHTLite TODO: Refactor globals!
+extern FManifest GManifest;
+extern ECompilationResult::Type GCompilationResult;
+extern FGraphEventArray GAsyncFileTasks;
+extern double GMacroizeTime;
+extern TArray<FString> ChangeMessages;
+extern bool bWriteContents;
+extern bool bVerifyContents;
+
+static void PerformSimplifiedClassParse(UPackage* InParent, const TCHAR* FileName, const TCHAR* Buffer, struct FPerHeaderData& PerHeaderData, FClassDeclarations& ClassDeclarations);
+static void ProcessInitialClassParse(FPerHeaderData& PerHeaderData, FTypeDefinitionInfoMap& TypeDefinitionInfoMap);
+
+#pragma region Plugins
+/** Get all script plugins based on ini setting */
+void GetScriptPlugins(TArray<IScriptGeneratorPluginInterface*>& ScriptPlugins)
+{
+	FScopedDurationTimer PluginTimeTracker(GPluginOverheadTime);
+
+	ScriptPlugins = IModularFeatures::Get().GetModularFeatureImplementations<IScriptGeneratorPluginInterface>(TEXT("ScriptGenerator"));
+	UE_LOG(LogCompile, Log, TEXT("Found %d script generator plugins."), ScriptPlugins.Num());
+
+	// Check if we can use these plugins and initialize them
+	for (int32 PluginIndex = ScriptPlugins.Num() - 1; PluginIndex >= 0; --PluginIndex)
+	{
+		IScriptGeneratorPluginInterface* ScriptGenerator = ScriptPlugins[PluginIndex];
+		bool bSupportedPlugin = ScriptGenerator->SupportsTarget(GManifest.TargetName);
+		if (bSupportedPlugin)
+		{
+			// Find the right output directory for this plugin base on its target (Engine-side) plugin name.
+			FString GeneratedCodeModuleName = ScriptGenerator->GetGeneratedCodeModuleName();
+			const FManifestModule* GeneratedCodeModule = NULL;
+			FString OutputDirectory;
+			FString IncludeBase;
+			for (const FManifestModule& Module : GManifest.Modules)
+			{
+				if (Module.Name == GeneratedCodeModuleName)
+				{
+					GeneratedCodeModule = &Module;
+				}
+			}
+			if (GeneratedCodeModule)
+			{
+				UE_LOG(LogCompile, Log, TEXT("Initializing script generator \'%s\'"), *ScriptGenerator->GetGeneratorName());
+				ScriptGenerator->Initialize(GManifest.RootLocalPath, GManifest.RootBuildPath, GeneratedCodeModule->GeneratedIncludeDirectory, GeneratedCodeModule->IncludeBase);
+			}
+			else
+			{
+				// Can't use this plugin
+				UE_LOG(LogCompile, Log, TEXT("Unable to determine output directory for %s. Cannot export script glue with \'%s\'"), *GeneratedCodeModuleName, *ScriptGenerator->GetGeneratorName());
+				bSupportedPlugin = false;
+			}
+		}
+		if (!bSupportedPlugin)
+		{
+			UE_LOG(LogCompile, Log, TEXT("Script generator \'%s\' not supported for target: %s"), *ScriptGenerator->GetGeneratorName(), *GManifest.TargetName);
+			ScriptPlugins.RemoveAt(PluginIndex);
+		}
+	}
+}
+#pragma endregion Plugins
+
+/**
+ * Tries to resolve super classes for classes defined in the given
+ * module.
+ *
+ * @param Package Modules package.
+ */
+void ResolveSuperClasses(UPackage* Package, const FTypeDefinitionInfoMap& TypeDefinitionInfoMap)
+{
+	TArray<UObject*> Objects;
+	GetObjectsWithPackage(Package, Objects);
+
+	for (UObject* Object : Objects)
+	{
+		if (!Object->IsA<UClass>() || Object->HasAnyFlags(RF_ClassDefaultObject))
+		{
+			continue;
+		}
+
+		UClass* DefinedClass = Cast<UClass>(Object);
+
+		if (DefinedClass->HasAnyClassFlags(CLASS_Intrinsic | CLASS_NoExport))
+		{
+			continue;
+		}
+
+		const FSimplifiedParsingClassInfo& ParsingInfo = TypeDefinitionInfoMap[DefinedClass]->GetUnrealSourceFile().GetDefinedClassParsingInfo(DefinedClass);
+
+		const FString& BaseClassName = ParsingInfo.GetBaseClassName();
+		const FString& BaseClassNameStripped = GetClassNameWithPrefixRemoved(BaseClassName);
+
+		if (!BaseClassNameStripped.IsEmpty() && !DefinedClass->GetSuperClass())
+		{
+			UClass* FoundBaseClass = FindObject<UClass>(Package, *BaseClassNameStripped);
+
+			if (FoundBaseClass == nullptr)
+			{
+				FoundBaseClass = FindObject<UClass>(ANY_PACKAGE, *BaseClassNameStripped);
+			}
+
+			if (FoundBaseClass == nullptr)
+			{
+				// Don't know its parent class. Raise error.
+				FError::Throwf(TEXT("Couldn't find parent type for '%s' named '%s' in current module (Package: %s) or any other module parsed so far."), *DefinedClass->GetName(), *BaseClassName, *GetNameSafe(Package));
+			}
+
+			DefinedClass->SetSuperStruct(FoundBaseClass);
+			DefinedClass->ClassCastFlags |= FoundBaseClass->ClassCastFlags;
+		}
+	}
+}
+
+ECompilationResult::Type PreparseModules(const FString& ModuleInfoPath, int32& NumFailures, FUnrealSourceFiles& UnrealSourceFilesMap, FTypeDefinitionInfoMap& TypeDefinitionInfoMap, FPublicSourceFileSet& PublicSourceFileSet, TMap<UPackage*, const FManifestModule*>& PackageToManifestModuleMap, FClassDeclarations& ClassDeclarations)
+{
+	// Three passes.  1) Public 'Classes' headers (legacy)  2) Public headers   3) Private headers
+	enum EHeaderFolderTypes
+	{
+		PublicClassesHeaders = 0,
+		PublicHeaders = 1,
+		PrivateHeaders,
+
+		FolderType_Count
+	};
+
+	ECompilationResult::Type Result = ECompilationResult::Succeeded;
+
+#if !PLATFORM_EXCEPTIONS_DISABLED
+	FGraphEventArray ExceptionTasks;
+
+	auto LogException = [&Result, &NumFailures, &ExceptionTasks](FString&& Filename, int32 Line, const FString& Message)
+	{
+		auto LogExceptionTask = [&Result, &NumFailures, Filename = MoveTemp(Filename), Line, Message]()
+		{
+			TGuardValue<ELogTimes::Type> DisableLogTimes(GPrintLogTimes, ELogTimes::None);
+
+			FString FormattedErrorMessage = FString::Printf(TEXT("%s(%d): Error: %s\r\n"), *Filename, Line, *Message);
+			Result = ECompilationResult::OtherCompilationError;
+
+			UE_LOG(LogCompile, Log, TEXT("%s"), *FormattedErrorMessage);
+			GWarn->Log(ELogVerbosity::Error, FormattedErrorMessage);
+
+			++NumFailures;
+		};
+
+		if (IsInGameThread())
+		{
+			LogExceptionTask();
+		}
+		else
+		{
+			FGraphEventRef EventRef = FFunctionGraphTask::CreateAndDispatchWhenReady(MoveTemp(LogExceptionTask), TStatId(), nullptr, ENamedThreads::GameThread);
+
+			static FCriticalSection ExceptionCS;
+			FScopeLock Lock(&ExceptionCS);
+			ExceptionTasks.Add(EventRef);
+		}
+	};
+#endif
+
+	for (FManifestModule& Module : GManifest.Modules)
+	{
+		if (Result != ECompilationResult::Succeeded)
+		{
+			break;
+		}
+
+		// Force regeneration of all subsequent modules, otherwise data will get corrupted.
+		Module.ForceRegeneration();
+
+		UPackage* Package = Cast<UPackage>(StaticFindObjectFast(UPackage::StaticClass(), NULL, FName(*Module.LongPackageName), false, false));
+		if (Package == NULL)
+		{
+			Package = CreatePackage(NULL, *Module.LongPackageName);
+		}
+		// Set some package flags for indicating that this package contains script
+		// NOTE: We do this even if we didn't have to create the package, because CoreUObject is compiled into UnrealHeaderTool and we still
+		//       want to make sure our flags get set
+		Package->SetPackageFlags(PKG_ContainsScript | PKG_Compiling);
+		Package->ClearPackageFlags(PKG_ClientOptional | PKG_ServerSideOnly);
+
+		if (Module.OverrideModuleType == EPackageOverrideType::None)
+		{
+			switch (Module.ModuleType)
+			{
+			case EBuildModuleType::GameEditor:
+			case EBuildModuleType::EngineEditor:
+				Package->SetPackageFlags(PKG_EditorOnly);
+				break;
+
+			case EBuildModuleType::GameDeveloper:
+			case EBuildModuleType::EngineDeveloper:
+				Package->SetPackageFlags(PKG_Developer);
+				break;
+
+			case EBuildModuleType::GameUncooked:
+			case EBuildModuleType::EngineUncooked:
+				Package->SetPackageFlags(PKG_UncookedOnly);
+				break;
+			}
+		}
+		else
+		{
+			// If the user has specified this module to have another package flag, then OR it on
+			switch (Module.OverrideModuleType)
+			{
+			case EPackageOverrideType::EditorOnly:
+				Package->SetPackageFlags(PKG_EditorOnly);
+				break;
+
+			case EPackageOverrideType::EngineDeveloper:
+			case EPackageOverrideType::GameDeveloper:
+				Package->SetPackageFlags(PKG_Developer);
+				break;
+
+			case EPackageOverrideType::EngineUncookedOnly:
+			case EPackageOverrideType::GameUncookedOnly:
+				Package->SetPackageFlags(PKG_UncookedOnly);
+				break;
+			}
+		}
+
+		// Add new module or overwrite whatever we had loaded, that data is obsolete.
+		PackageToManifestModuleMap.Add(Package, &Module);
+
+		double ThisModulePreparseTime = 0.0;
+		int32 NumHeadersPreparsed = 0;
+		FDurationTimer ThisModuleTimer(ThisModulePreparseTime);
+		ThisModuleTimer.Start();
+
+		// Pre-parse the headers
+		for (int32 PassIndex = 0; PassIndex < FolderType_Count && Result == ECompilationResult::Succeeded; ++PassIndex)
+		{
+			EHeaderFolderTypes CurrentlyProcessing = (EHeaderFolderTypes)PassIndex;
+
+			// We'll make an ordered list of all UObject headers we care about.
+			// @todo uht: Ideally 'dependson' would not be allowed from public -> private, or NOT at all for new style headers
+			const TArray<FString>& UObjectHeaders =
+				(CurrentlyProcessing == PublicClassesHeaders) ? Module.PublicUObjectClassesHeaders :
+				(CurrentlyProcessing == PublicHeaders) ? Module.PublicUObjectHeaders :
+				Module.PrivateUObjectHeaders;
+			if (!UObjectHeaders.Num())
+			{
+				continue;
+			}
+
+			NumHeadersPreparsed += UObjectHeaders.Num();
+
+			TArray<FString> HeaderFiles;
+			HeaderFiles.SetNum(UObjectHeaders.Num());
+
+			{
+				SCOPE_SECONDS_COUNTER_UHT(LoadHeaderContentFromFile);
+#if	UHT_USE_PARALLEL_FOR
+				ParallelFor(UObjectHeaders.Num(), [&](int32 Index)
+#else
+				for (int Index = 0; Index < UObjectHeaders.Num(); Index++)
+#endif
+				{
+					const FString& RawFilename = UObjectHeaders[Index];
+
+#if !PLATFORM_EXCEPTIONS_DISABLED
+					try
+#endif
+					{
+						const FString FullFilename = FPaths::ConvertRelativePathToFull(ModuleInfoPath, RawFilename);
+
+						if (!FFileHelper::LoadFileToString(HeaderFiles[Index], *FullFilename))
+						{
+							FError::Throwf(TEXT("UnrealHeaderTool was unable to load source file '%s'"), *FullFilename);
+						}
+					}
+#if !PLATFORM_EXCEPTIONS_DISABLED
+					catch (TCHAR* ErrorMsg)
+					{
+						FString AbsFilename = IFileManager::Get().ConvertToAbsolutePathForExternalAppForRead(*RawFilename);
+						LogException(MoveTemp(AbsFilename), 1, ErrorMsg);
+					}
+#endif
+				}
+#if	UHT_USE_PARALLEL_FOR
+				);
+#endif
+			}
+
+#if !PLATFORM_EXCEPTIONS_DISABLED
+			FTaskGraphInterface::Get().WaitUntilTasksComplete(ExceptionTasks);
+#endif
+
+			if (Result != ECompilationResult::Succeeded)
+			{
+				continue;
+			}
+
+			TArray<FPerHeaderData> PerHeaderData;
+			PerHeaderData.SetNum(UObjectHeaders.Num());
+
+#if	UHT_USE_PARALLEL_FOR
+			ParallelFor(UObjectHeaders.Num(), [&](int32 Index)
+#else
+			for (int Index = 0; Index < UObjectHeaders.Num(); Index++)
+#endif
+			{
+				const FString& RawFilename = UObjectHeaders[Index];
+
+#if !PLATFORM_EXCEPTIONS_DISABLED
+				try
+#endif
+				{
+					PerformSimplifiedClassParse(Package, *RawFilename, *HeaderFiles[Index], PerHeaderData[Index], ClassDeclarations);
+				}
+#if !PLATFORM_EXCEPTIONS_DISABLED
+				catch (const FFileLineException& Ex)
+				{
+					FString AbsFilename = IFileManager::Get().ConvertToAbsolutePathForExternalAppForRead(*Ex.Filename);
+					LogException(MoveTemp(AbsFilename), Ex.Line, Ex.Message);
+				}
+				catch (TCHAR* ErrorMsg)
+				{
+					FString AbsFilename = IFileManager::Get().ConvertToAbsolutePathForExternalAppForRead(*RawFilename);
+					LogException(MoveTemp(AbsFilename), 1, ErrorMsg);
+				}
+#endif
+			}
+#if	UHT_USE_PARALLEL_FOR
+			);
+#endif
+
+#if !PLATFORM_EXCEPTIONS_DISABLED
+			FTaskGraphInterface::Get().WaitUntilTasksComplete(ExceptionTasks);
+#endif
+
+			if (Result != ECompilationResult::Succeeded)
+			{
+				continue;
+			}
+
+			for (int32 Index = 0; Index < UObjectHeaders.Num(); ++Index)
+			{
+				const FString& RawFilename = UObjectHeaders[Index];
+
+#if !PLATFORM_EXCEPTIONS_DISABLED
+				try
+#endif
+				{
+					// Import class.
+					const FString FullFilename = FPaths::ConvertRelativePathToFull(ModuleInfoPath, RawFilename);
+
+					ProcessInitialClassParse(PerHeaderData[Index], TypeDefinitionInfoMap);
+					TSharedRef<FUnrealSourceFile> UnrealSourceFile = PerHeaderData[Index].UnrealSourceFile.ToSharedRef();
+					FUnrealSourceFile* UnrealSourceFilePtr = &UnrealSourceFile.Get();
+					UnrealSourceFilesMap.Add(FPaths::GetCleanFilename(RawFilename), UnrealSourceFile);
+
+					if (CurrentlyProcessing == PublicClassesHeaders)
+					{
+						PublicSourceFileSet.Add(UnrealSourceFilePtr);
+					}
+
+					// Save metadata for the class path, both for it's include path and relative to the module base directory
+					if (FullFilename.StartsWith(Module.BaseDirectory))
+					{
+						// Get the path relative to the module directory
+						const TCHAR* ModuleRelativePath = *FullFilename + Module.BaseDirectory.Len();
+
+						UnrealSourceFile->SetModuleRelativePath(ModuleRelativePath);
+
+						// Calculate the include path
+						const TCHAR* IncludePath = ModuleRelativePath;
+
+						// Walk over the first potential slash
+						if (*IncludePath == TEXT('/'))
+						{
+							IncludePath++;
+						}
+
+						// Does this module path start with a known include path location? If so, we can cut that part out of the include path
+						static const TCHAR PublicFolderName[] = TEXT("Public/");
+						static const TCHAR PrivateFolderName[] = TEXT("Private/");
+						static const TCHAR ClassesFolderName[] = TEXT("Classes/");
+						if (FCString::Strnicmp(IncludePath, PublicFolderName, UE_ARRAY_COUNT(PublicFolderName) - 1) == 0)
+						{
+							IncludePath += (UE_ARRAY_COUNT(PublicFolderName) - 1);
+						}
+						else if (FCString::Strnicmp(IncludePath, PrivateFolderName, UE_ARRAY_COUNT(PrivateFolderName) - 1) == 0)
+						{
+							IncludePath += (UE_ARRAY_COUNT(PrivateFolderName) - 1);
+						}
+						else if (FCString::Strnicmp(IncludePath, ClassesFolderName, UE_ARRAY_COUNT(ClassesFolderName) - 1) == 0)
+						{
+							IncludePath += (UE_ARRAY_COUNT(ClassesFolderName) - 1);
+						}
+
+						// Add the include path
+						if (*IncludePath != 0)
+						{
+							UnrealSourceFile->SetIncludePath(MoveTemp(IncludePath));
+						}
+					}
+				}
+#if !PLATFORM_EXCEPTIONS_DISABLED
+				catch (const FFileLineException& Ex)
+				{
+					FString AbsFilename = IFileManager::Get().ConvertToAbsolutePathForExternalAppForRead(*Ex.Filename);
+					LogException(MoveTemp(AbsFilename), Ex.Line, Ex.Message);
+				}
+				catch (TCHAR* ErrorMsg)
+				{
+					FString AbsFilename = IFileManager::Get().ConvertToAbsolutePathForExternalAppForRead(*RawFilename);
+					LogException(MoveTemp(AbsFilename), 1, ErrorMsg);
+				}
+#endif
+			}
+			if (Result == ECompilationResult::Succeeded && NumFailures != 0)
+			{
+				Result = ECompilationResult::OtherCompilationError;
+			}
+		}
+
+		// Don't resolve superclasses for module when loading from makefile.
+		// Data is only partially loaded at this point.
+#if !PLATFORM_EXCEPTIONS_DISABLED
+		try
+#endif
+		{
+			ResolveSuperClasses(Package, TypeDefinitionInfoMap);
+		}
+#if !PLATFORM_EXCEPTIONS_DISABLED
+		catch (TCHAR* ErrorMsg)
+		{
+			TGuardValue<ELogTimes::Type> DisableLogTimes(GPrintLogTimes, ELogTimes::None);
+
+			FString FormattedErrorMessage = FString::Printf(TEXT("Error: %s\r\n"), ErrorMsg);
+
+			Result = GCompilationResult;
+
+			UE_LOG(LogCompile, Log, TEXT("%s"), *FormattedErrorMessage);
+			GWarn->Log(ELogVerbosity::Error, FormattedErrorMessage);
+
+			++NumFailures;
+		}
+#endif
+
+		ThisModuleTimer.Stop();
+		UE_LOG(LogCompile, Log, TEXT("Preparsed module %s containing %i files(s) in %.2f secs."), *Module.LongPackageName, NumHeadersPreparsed, ThisModulePreparseTime);
+	}
+
+	return Result;
+}
+
+ECompilationResult::Type UnrealHeaderTool_Main(const FString& ModuleInfoFilename)
+{
+	double MainTime = 0.0;
+	FDurationTimer MainTimer(MainTime);
+	MainTimer.Start();
+
+	check(GIsUCCMakeStandaloneHeaderGenerator);
+	ECompilationResult::Type Result = ECompilationResult::Succeeded;
+
+	FString ModuleInfoPath = FPaths::GetPath(ModuleInfoFilename);
+
+	// Load the manifest file, giving a list of all modules to be processed, pre-sorted by dependency ordering
+#if !PLATFORM_EXCEPTIONS_DISABLED
+	try
+#endif
+	{
+		GManifest = FManifest::LoadFromFile(ModuleInfoFilename);
+	}
+#if !PLATFORM_EXCEPTIONS_DISABLED
+	catch (const TCHAR* Ex)
+	{
+		UE_LOG(LogCompile, Error, TEXT("Failed to load manifest file '%s': %s"), *ModuleInfoFilename, Ex);
+		return GCompilationResult;
+	}
+#endif
+
+	// Counters.
+	int32 NumFailures = 0;
+	double TotalModulePreparseTime = 0.0;
+	double TotalParseAndCodegenTime = 0.0;
+
+	// UHTLite NOTE: Used in PreparseModules() and UnrealHeaderTool_Main().
+	// UHTLite NOTE: Written to in FHeaderParser::CompileStructDeclaration() and FHeaderParser::CompileEnum().
+	// UHTLite NOTE: Needs a mutex?
+	FTypeDefinitionInfoMap TypeDefinitionInfoMap;
+
+	// UHTLite NOTE: Data is written once in PreparseModules() and is then read-only.
+	FUnrealSourceFiles   UnrealSourceFilesMap;
+	FPublicSourceFileSet PublicSourceFileSet;
+	TMap<UPackage*, const FManifestModule*> PackageToManifestModuleMap;
+
+	// UHTLite NOTE: Classes are added in PreparseModules(), but class flags are modified during header parsing!
+	FClassDeclarations ClassDeclarations;
+
+	// UHTLite NOTE: This data is cross-module! It is read-only during code generation.
+	TMap<UEnum*, EUnderlyingEnumType>      EnumUnderlyingTypes;
+	FRWLock                                EnumUnderlyingTypesLock;
+	TMap<UClass*, FArchiveTypeDefinePair>  ClassSerializerMap;
+	FRWLock                                ClassSerializerMapLock;
+
+	{
+		FDurationTimer TotalModulePreparseTimer(TotalModulePreparseTime);
+		TotalModulePreparseTimer.Start();
+		Result = PreparseModules(ModuleInfoPath, NumFailures, UnrealSourceFilesMap, TypeDefinitionInfoMap, PublicSourceFileSet, PackageToManifestModuleMap, ClassDeclarations);
+		TotalModulePreparseTimer.Stop();
+	}
+	// Do the actual parse of the headers and generate for them
+	if (Result == ECompilationResult::Succeeded)
+	{
+		FScopedDurationTimer ParseAndCodeGenTimer(TotalParseAndCodegenTime);
+
+		TMap<UPackage*, TArray<UClass*>> ClassesByPackageMap;
+		ClassesByPackageMap.Reserve(GManifest.Modules.Num());
+
+		// Verify that all script declared superclasses exist.
+		for (UClass* ScriptClass : TObjectRange<UClass>())
+		{
+			ClassesByPackageMap.FindOrAdd(ScriptClass->GetOutermost()).Add(ScriptClass);
+
+			// UHTLite NOTE: Not required for code-gen. Will be refactored later.
+			/*
+			const UClass* ScriptSuperClass = ScriptClass->GetSuperClass();
+
+			if (ScriptSuperClass && !ScriptSuperClass->HasAnyClassFlags(CLASS_Intrinsic) && TypeDefinitionInfoMap.Contains(ScriptClass) && !TypeDefinitionInfoMap.Contains(ScriptSuperClass))
+			{
+				class FSuperClassContextSupplier : public FContextSupplier
+				{
+				public:
+					FSuperClassContextSupplier(const UClass* Class)
+						: DefinitionInfo(TypeDefinitionInfoMap[Class])
+					{ }
+
+					virtual FString GetContext() override
+					{
+						FString Filename = IFileManager::Get().ConvertToAbsolutePathForExternalAppForRead(*DefinitionInfo->GetUnrealSourceFile().GetFilename());
+						int32 LineNumber = DefinitionInfo->GetLineNumber();
+						return FString::Printf(TEXT("%s(%i)"), *Filename, LineNumber);
+					}
+				private:
+					TSharedRef<FUnrealTypeDefinitionInfo> DefinitionInfo;
+				} ContextSupplier(ScriptClass);
+
+				FContextSupplier* OldContext = GWarn->GetContext();
+
+				TGuardValue<ELogTimes::Type> DisableLogTimes(GPrintLogTimes, ELogTimes::None);
+
+				GWarn->SetContext(&ContextSupplier);
+				GWarn->Log(ELogVerbosity::Error, FString::Printf(TEXT("Error: Superclass %s of class %s not found"), *ScriptSuperClass->GetName(), *ScriptClass->GetName()));
+				GWarn->SetContext(OldContext);
+
+				Result = ECompilationResult::OtherCompilationError;
+				++NumFailures;
+			}
+			*/
+		}
+
+		if (Result == ECompilationResult::Succeeded)
+		{
+#pragma region Plugins
+			TArray<IScriptGeneratorPluginInterface*> ScriptPlugins;
+			// Can only export scripts for game targets
+			if (GManifest.IsGameTarget)
+			{
+				GetScriptPlugins(ScriptPlugins);
+			}
+#pragma endregion Plugins
+
+			// UHTLite NOTE: TypeDefinitionInfoMap needs a mutex before this loop can be threaded!
+			for (const FManifestModule& Module : GManifest.Modules)
+			{
+				if (UPackage* Package = Cast<UPackage>(StaticFindObjectFast(UPackage::StaticClass(), NULL, FName(*Module.LongPackageName), false, false)))
+				{
+					FClasses AllClasses(ClassesByPackageMap.Find(Package));
+					AllClasses.Validate();
+
+					Result = FHeaderParser::ParseAllHeadersInside(AllClasses,
+						GWarn,
+						Package,
+						Module,
+#pragma region Plugins
+						ScriptPlugins,
+#pragma endregion Plugins
+						UnrealSourceFilesMap,
+						TypeDefinitionInfoMap,
+						PublicSourceFileSet,
+						PackageToManifestModuleMap,
+						ClassDeclarations,
+						EnumUnderlyingTypes,
+						EnumUnderlyingTypesLock,
+						ClassSerializerMap,
+						ClassSerializerMapLock);
+
+					if (Result != ECompilationResult::Succeeded)
+					{
+						++NumFailures;
+						break;
+					}
+				}
+			}
+
+#pragma region Plugins
+			{
+				FScopedDurationTimer PluginTimeTracker(GPluginOverheadTime);
+				for (IScriptGeneratorPluginInterface* ScriptGenerator : ScriptPlugins)
+				{
+					ScriptGenerator->FinishExport();
+				}
+			}
+
+			// Get a list of external dependencies from each enabled plugin
+			FString ExternalDependencies;
+			for (IScriptGeneratorPluginInterface* ScriptPlugin : ScriptPlugins)
+			{
+				TArray<FString> PluginExternalDependencies;
+				ScriptPlugin->GetExternalDependencies(PluginExternalDependencies);
+
+				for (const FString& PluginExternalDependency : PluginExternalDependencies)
+				{
+					ExternalDependencies += PluginExternalDependency + LINE_TERMINATOR;
+				}
+			}
+
+			FFileHelper::SaveStringToFile(ExternalDependencies, *GManifest.ExternalDependenciesFile);
+#pragma endregion Plugins
+		}
+	}
+
+	// Avoid TArray slack for meta data.
+	GScriptHelper.Shrink();
+
+	// Finish all async file tasks before stopping the clock
+	FTaskGraphInterface::Get().WaitUntilTasksComplete(GAsyncFileTasks);
+
+	MainTimer.Stop();
+
+	UE_LOG(LogCompile, Log, TEXT("Preparsing %i modules took %.2f seconds"), GManifest.Modules.Num(), TotalModulePreparseTime);
+	UE_LOG(LogCompile, Log, TEXT("Parsing took %.2f seconds"), TotalParseAndCodegenTime - GHeaderCodeGenTime);
+	UE_LOG(LogCompile, Log, TEXT("Code generation took %.2f seconds"), GHeaderCodeGenTime);
+	UE_LOG(LogCompile, Log, TEXT("ScriptPlugin overhead was %.2f seconds"), GPluginOverheadTime);
+	UE_LOG(LogCompile, Log, TEXT("Macroize time was %.2f seconds"), GMacroizeTime);
+
+	FUnrealHeaderToolStats& Stats = FUnrealHeaderToolStats::Get();
+	for (const TPair<FName, double>& Pair : Stats.Counters)
+	{
+		FString CounterName = Pair.Key.ToString();
+		UE_LOG(LogCompile, Log, TEXT("%s timer was %.3f seconds"), *CounterName, Pair.Value);
+	}
+
+	UE_LOG(LogCompile, Log, TEXT("Total time was %.2f seconds"), MainTime);
+
+	if (bWriteContents)
+	{
+		UE_LOG(LogCompile, Log, TEXT("********************************* Wrote reference generated code to ReferenceGeneratedCode."));
+	}
+	else if (bVerifyContents)
+	{
+		UE_LOG(LogCompile, Log, TEXT("********************************* Wrote generated code to VerifyGeneratedCode and compared to ReferenceGeneratedCode"));
+		for (FString& Msg : ChangeMessages)
+		{
+			UE_LOG(LogCompile, Error, TEXT("%s"), *Msg);
+		}
+		TArray<FString> RefFileNames;
+		IFileManager::Get().FindFiles(RefFileNames, *(FPaths::ProjectSavedDir() / TEXT("ReferenceGeneratedCode/*.*")), true, false);
+		TArray<FString> VerFileNames;
+		IFileManager::Get().FindFiles(VerFileNames, *(FPaths::ProjectSavedDir() / TEXT("VerifyGeneratedCode/*.*")), true, false);
+		if (RefFileNames.Num() != VerFileNames.Num())
+		{
+			UE_LOG(LogCompile, Error, TEXT("Number of generated files mismatch ref=%d, ver=%d"), RefFileNames.Num(), VerFileNames.Num());
+		}
+	}
+
+	RequestEngineExit(TEXT("UnrealHeaderTool finished"));
+
+	if (Result != ECompilationResult::Succeeded || NumFailures > 0)
+	{
+		return ECompilationResult::OtherCompilationError;
+	}
+
+	return Result;
+}
+
+UClass* ProcessParsedClass(bool bClassIsAnInterface, TArray<FHeaderProvider>& DependentOn, const FString& ClassName, const FString& BaseClassName, UObject* InParent, EObjectFlags Flags)
+{
+	FString ClassNameStripped = GetClassNameWithPrefixRemoved(*ClassName);
+
+	// All classes must start with a valid unreal prefix
+	if (!FHeaderParser::ClassNameHasValidPrefix(ClassName, ClassNameStripped))
+	{
+		FError::Throwf(TEXT("Invalid class name '%s'. The class name must have an appropriate prefix added (A for Actors, U for other classes)."), *ClassName);
+	}
+
+	if (FHeaderParser::IsReservedTypeName(ClassNameStripped))
+	{
+		FError::Throwf(TEXT("Invalid class name '%s'. Cannot use a reserved name ('%s')."), *ClassName, *ClassNameStripped);
+	}
+
+	// Ensure the base class has any valid prefix and exists as a valid class. Checking for the 'correct' prefix will occur during compilation
+	FString BaseClassNameStripped;
+	if (!BaseClassName.IsEmpty())
+	{
+		BaseClassNameStripped = GetClassNameWithPrefixRemoved(BaseClassName);
+		if (!FHeaderParser::ClassNameHasValidPrefix(BaseClassName, BaseClassNameStripped))
+		{
+			FError::Throwf(TEXT("No prefix or invalid identifier for base class %s.\nClass names must match Unreal prefix specifications (e.g., \"UObject\" or \"AActor\")"), *BaseClassName);
+		}
+	}
+
+	//UE_LOG(LogCompile, Log, TEXT("Class: %s extends %s"),*ClassName,*BaseClassName);
+	// Handle failure and non-class headers.
+	if (BaseClassName.IsEmpty() && (ClassName != TEXT("UObject")))
+	{
+		FError::Throwf(TEXT("Class '%s' must inherit UObject or a UObject-derived class"), *ClassName);
+	}
+
+	if (ClassName == BaseClassName)
+	{
+		FError::Throwf(TEXT("Class '%s' cannot inherit from itself"), *ClassName);
+	}
+
+	// In case the file system and the class disagree on the case of the
+	// class name replace the fname with the one from the script class file
+	// This is needed because not all source control systems respect the
+	// original filename's case
+	FName ClassNameReplace(*ClassName, FNAME_Replace_Not_Safe_For_Threading);
+
+	// Use stripped class name for processing and replace as we did above
+	FName ClassNameStrippedReplace(*ClassNameStripped, FNAME_Replace_Not_Safe_For_Threading);
+
+	UClass* ResultClass = FindObject<UClass>(InParent, *ClassNameStripped);
+
+	// if we aren't generating headers, then we shouldn't set misaligned object, since it won't get cleared
+
+	const static bool bVerboseOutput = FParse::Param(FCommandLine::Get(), TEXT("VERBOSE"));
+
+	if (ResultClass == nullptr || !ResultClass->IsNative())
+	{
+		// detect if the same class name is used in multiple packages
+		if (ResultClass == nullptr)
+		{
+			UClass* ConflictingClass = FindObject<UClass>(ANY_PACKAGE, *ClassNameStripped, true);
+			if (ConflictingClass != nullptr)
+			{
+				UE_LOG_WARNING_UHT(TEXT("Duplicate class name: %s also exists in file %s"), *ClassName, *ConflictingClass->GetOutermost()->GetName());
+			}
+		}
+
+		// Create new class.
+		ResultClass = new(EC_InternalUseOnlyConstructor, InParent, *ClassNameStripped, Flags) UClass(FObjectInitializer(), nullptr);
+
+		// add CLASS_Interface flag if the class is an interface
+		// NOTE: at this pre-parsing/importing stage, we cannot know if our super class is an interface or not,
+		// we leave the validation to the main header parser
+		if (bClassIsAnInterface)
+		{
+			ResultClass->ClassFlags |= CLASS_Interface;
+		}
+
+		if (bVerboseOutput)
+		{
+			UE_LOG(LogCompile, Log, TEXT("Imported: %s"), *ResultClass->GetFullName());
+		}
+	}
+
+	if (bVerboseOutput)
+	{
+		for (const FHeaderProvider& Dependency : DependentOn)
+		{
+			UE_LOG(LogCompile, Log, TEXT("\tAdding %s as a dependency"), *Dependency.ToString());
+		}
+	}
+
+	return ResultClass;
+}
+
+void PerformSimplifiedClassParse(UPackage* InParent, const TCHAR* FileName, const TCHAR* Buffer, FPerHeaderData& PerHeaderData, FClassDeclarations& ClassDeclarations)
+{
+	// Parse the header to extract the information needed
+	FUHTStringBuilder ClassHeaderTextStrippedOfCppText;
+
+	FHeaderParser::SimplifiedClassParse(FileName, Buffer, /*out*/ PerHeaderData.ParsedClassArray, /*out*/ PerHeaderData.DependsOn, ClassHeaderTextStrippedOfCppText, ClassDeclarations);
+
+	FUnrealSourceFile* UnrealSourceFilePtr = new FUnrealSourceFile(InParent, FileName, MoveTemp(ClassHeaderTextStrippedOfCppText));
+	PerHeaderData.UnrealSourceFile = MakeShareable(UnrealSourceFilePtr);
+}
+
+void ProcessInitialClassParse(FPerHeaderData& PerHeaderData, FTypeDefinitionInfoMap& TypeDefinitionInfoMap)
+{
+	TSharedRef<FUnrealSourceFile> UnrealSourceFile = PerHeaderData.UnrealSourceFile.ToSharedRef();
+	UPackage* InParent = UnrealSourceFile->GetPackage();
+	for (FSimplifiedParsingClassInfo& ParsedClassInfo : PerHeaderData.ParsedClassArray)
+	{
+		UClass* ResultClass = ProcessParsedClass(ParsedClassInfo.IsInterface(), PerHeaderData.DependsOn, ParsedClassInfo.GetClassName(), ParsedClassInfo.GetBaseClassName(), InParent, RF_Public | RF_Standalone);
+
+		// UHTLite NOTE: Not required for code-gen. Will be refactored later.
+		/*
+		GStructToSourceLine.Add(ResultClass, MakeTuple(UnrealSourceFile, ParsedClassInfo.GetClassDefLine()));
+		*/
+
+		FScope::AddTypeScope(ResultClass, &UnrealSourceFile->GetScope().Get());
+
+		TypeDefinitionInfoMap.Add(ResultClass, MakeShared<FUnrealTypeDefinitionInfo>(UnrealSourceFile.Get(), ParsedClassInfo.GetClassDefLine()));
+		UnrealSourceFile->AddDefinedClass(ResultClass, MoveTemp(ParsedClassInfo));
+	}
+
+	for (FHeaderProvider& DependsOnElement : PerHeaderData.DependsOn)
+	{
+		UnrealSourceFile->GetIncludes().AddUnique(DependsOnElement);
+	}
+}
+
 /**
  * Application entry point
  *
@@ -108,7 +926,6 @@ INT32_MAIN_INT32_ARGC_TCHAR_ARGV()
 		ModuleInfoFilename = FPaths::ConvertRelativePathToFull(FPlatformProcess::BaseDir(), TEXT("../../Source/Programs/UnrealHeaderTool/Resources/UHTDebugging.manifest"));
 	}
 
-	extern ECompilationResult::Type UnrealHeaderTool_Main(const FString& ModuleInfoFilename);
 	ECompilationResult::Type Result = UnrealHeaderTool_Main(ModuleInfoFilename);
 
 	if (Result == ECompilationResult::Succeeded && (GUHTErrorLogged || (GUHTWarningLogged && GWarn->TreatWarningsAsErrors)))
@@ -118,4 +935,3 @@ INT32_MAIN_INT32_ARGC_TCHAR_ARGV()
 
 	return Result;
 }
-
diff --git a/Engine/Source/Programs/UHTLite/Private/UnrealSourceFile.cpp b/Engine/Source/Programs/UHTLite/Private/UnrealSourceFile.cpp
--- a/Engine/Source/Programs/UHTLite/Private/UnrealSourceFile.cpp
+++ b/Engine/Source/Programs/UHTLite/Private/UnrealSourceFile.cpp
@@ -109,14 +109,14 @@ const FString& FUnrealSourceFile::GetContent() const
 	return Content;
 }
 
-EGeneratedCodeVersion FUnrealSourceFile::GetGeneratedCodeVersionForStruct(UStruct* Struct) const
+EGeneratedCodeVersion FUnrealSourceFile::GetGeneratedCodeVersionForStruct(const FHeaderParser& HeaderParser, UStruct* Struct) const
 {
 	if (const EGeneratedCodeVersion* Version = GeneratedCodeVersions.Find(Struct))
 	{
 		return *Version;
 	}
 
-	return FHeaderParser::DefaultGeneratedCodeVersion;
+	return HeaderParser.DefaultGeneratedCodeVersion;
 }
 
 void FUnrealSourceFile::MarkDependenciesResolved()
diff --git a/Engine/Source/Programs/UHTLite/Private/UnrealSourceFile.h b/Engine/Source/Programs/UHTLite/Private/UnrealSourceFile.h
--- a/Engine/Source/Programs/UHTLite/Private/UnrealSourceFile.h
+++ b/Engine/Source/Programs/UHTLite/Private/UnrealSourceFile.h
@@ -197,7 +197,7 @@ public:
 	/**
 	 * Gets generated code version for given UStruct.
 	 */
-	EGeneratedCodeVersion GetGeneratedCodeVersionForStruct(UStruct* Struct) const;
+	EGeneratedCodeVersion GetGeneratedCodeVersionForStruct(const class FHeaderParser& HeaderParser, UStruct* Struct) const;
 
 	/**
 	 * Gets generated code versions.
