build fix #rb marc.audy #okforgithub public

diff --git a/Engine/Source/Programs/UHTLite/Private/HeaderParser.cpp b/Engine/Source/Programs/UHTLite/Private/HeaderParser.cpp
--- a/Engine/Source/Programs/UHTLite/Private/HeaderParser.cpp
+++ b/Engine/Source/Programs/UHTLite/Private/HeaderParser.cpp
@@ -1895,7 +1895,7 @@ UEnum* FHeaderParser::CompileEnum()
 	}
 
 	// Register the list of enum names.
-	if (!Enum->SetEnums(EnumNames, CppForm, false))
+	if (!Enum->SetEnums(EnumNames, CppForm, EEnumFlags::None, false))
 	{
 		const FName MaxEnumItem      = *(Enum->GenerateEnumPrefix() + TEXT("_MAX"));
 		const int32 MaxEnumItemIndex = Enum->GetIndexByName(MaxEnumItem);
