No longer generate PPO offsets since nativization is no longer a thing

diff --git a/Engine/Source/Programs/UHTLite/Private/CodeGenerator.cpp b/Engine/Source/Programs/UHTLite/Private/CodeGenerator.cpp
--- a/Engine/Source/Programs/UHTLite/Private/CodeGenerator.cpp
+++ b/Engine/Source/Programs/UHTLite/Private/CodeGenerator.cpp
@@ -2926,34 +2926,6 @@ void WriteMacro(FOutputDevice& Output, const FString& MacroName, FString MacroCo
 	Output.Log(Macroize(*MacroName, MoveTemp(MacroContent)));
 }
 
-static FString PrivatePropertiesOffsetGetters(const UStruct* Struct, const FString& StructCppName)
-{
-	check(Struct);
-
-	FUHTStringBuilder Result;
-	for (const FProperty* Property : TFieldRange<FProperty>(Struct, EFieldIteratorFlags::ExcludeSuper))
-	{
-		if (Property && Property->HasAnyPropertyFlags(CPF_NativeAccessSpecifierPrivate | CPF_NativeAccessSpecifierProtected) && !Property->HasAnyPropertyFlags(CPF_EditorOnly))
-		{
-			const FBoolProperty* BoolProperty = CastField<const FBoolProperty>(Property);
-			if (BoolProperty && !BoolProperty->IsNativeBool()) // if it's a bitfield
-			{
-				continue;
-			}
-
-			FString PropertyName = Property->GetName();
-			if (Property->HasAllPropertyFlags(CPF_Deprecated))
-			{
-				PropertyName += TEXT("_DEPRECATED");
-			}
-			Result.Logf(TEXT("\tFORCEINLINE static uint32 __PPO__%s() { return STRUCT_OFFSET(%s, %s); }") LINE_TERMINATOR,
-				*PropertyName, *StructCppName, *PropertyName);
-		}
-	}
-
-	return MoveTemp(Result);
-}
-
 void FNativeClassHeaderGenerator::ExportClassFromSourceFileInner(
 	FOutputDevice&           OutGeneratedHeaderText,
 	FOutputDevice&           OutCpp,
@@ -3041,8 +3013,6 @@ void FNativeClassHeaderGenerator::ExportClassFromSourceFileInner(
 		APIArg = TEXT("NO");
 	}
 
-	FString PPOMacroName;
-
 	ClassDefinitionRange ClassRange;
 	if (const ClassDefinitionRange* FoundRange = HeaderParser.ClassDefinitionRanges.Find(Class))
 	{
@@ -3244,12 +3214,6 @@ void FNativeClassHeaderGenerator::ExportClassFromSourceFileInner(
 
 				ExportConstructorsMacros(OutGeneratedHeaderText, OutCpp, StandardUObjectConstructorsMacroCall, EnhancedUObjectConstructorsMacroCall, SourceFile.GetGeneratedMacroName(ClassData), Class, *APIArg);
 			}
-			{
-				const FString PrivatePropertiesOffsets = PrivatePropertiesOffsetGetters(Class, ClassCPPName);
-				const FString PPOMacroNameRaw = SourceFile.GetGeneratedMacroName(ClassData, TEXT("_PRIVATE_PROPERTY_OFFSET"));
-				PPOMacroName = FString::Printf(TEXT("\t%s\r\n"), *PPOMacroNameRaw);
-				WriteMacro(OutGeneratedHeaderText, PPOMacroNameRaw, PrivatePropertiesOffsets);
-			}
 		}
 	}
 
@@ -3282,8 +3246,8 @@ void FNativeClassHeaderGenerator::ExportClassFromSourceFileInner(
 			MacroName = TEXT("GENERATED_UCLASS_BODY()");
 			DeprecationWarning = GetGeneratedMacroDeprecationWarning(MacroName);
 			GeneratedBodyLine = ClassData->GetGeneratedBodyLine();
-			LegacyGeneratedBody = FString::Printf(TEXT("%s%s%s"), *PPOMacroName, *ClassMacroCalls, *StandardUObjectConstructorsMacroCall);
-			GeneratedBody = FString::Printf(TEXT("%s%s%s"), *PPOMacroName, *ClassNoPureDeclsMacroCalls, *EnhancedUObjectConstructorsMacroCall);
+			LegacyGeneratedBody = FString::Printf(TEXT("%s%s"), *ClassMacroCalls, *StandardUObjectConstructorsMacroCall);
+			GeneratedBody = FString::Printf(TEXT("%s%s"), *ClassNoPureDeclsMacroCalls, *EnhancedUObjectConstructorsMacroCall);
 		}
 
 		FString WrappedLegacyGeneratedBody = FString::Printf(TEXT("%s%s%s%s%s%s"), *DeprecationWarning, DisableDeprecationWarnings, Public, *LegacyGeneratedBody, Public, EnableDeprecationWarnings);
@@ -3712,7 +3676,6 @@ void FNativeClassHeaderGenerator::ExportGeneratedStructBodyMacros(FOutputDevice&
 
 		const FString FriendLine = FString::Printf(TEXT("\tfriend struct %s_Statics;\r\n"), *ChoppedSingletonName);
 		const FString StaticClassLine = FString::Printf(TEXT("\t%sstatic class UScriptStruct* StaticStruct();\r\n"), (bRequiredAPI ? *FriendApiString : TEXT("")));
-		const FString PrivatePropertiesOffset = PrivatePropertiesOffsetGetters(Struct, StructNameCPP);
 		
 		// if we have RigVM methods on this struct we need to 
 		// declare the static method as well as the stub method
@@ -3769,7 +3732,7 @@ void FNativeClassHeaderGenerator::ExportGeneratedStructBodyMacros(FOutputDevice&
 
 		const FString SuperTypedef = BaseStruct ? FString::Printf(TEXT("\ttypedef %s Super;\r\n"), *FNameLookupCPP::GetNameCPP(BaseStruct)) : FString();
 
-		FString CombinedLine = FString::Printf(TEXT("%s%s%s%s%s"), *FriendLine, *StaticClassLine, *RigVMMethodsDeclarations, *PrivatePropertiesOffset, *SuperTypedef);
+		FString CombinedLine = FString::Printf(TEXT("%s%s%s%s"), *FriendLine, *StaticClassLine, *RigVMMethodsDeclarations, *SuperTypedef);
 		const FString MacroName = SourceFile.GetGeneratedBodyMacroName(Struct->StructMacroDeclaredLineNumber);
 
 		const FString Macroized = Macroize(*MacroName, MoveTemp(CombinedLine));
diff --git a/Engine/Source/Programs/UnrealHeaderTool/Private/CodeGenerator.cpp b/Engine/Source/Programs/UnrealHeaderTool/Private/CodeGenerator.cpp
--- a/Engine/Source/Programs/UnrealHeaderTool/Private/CodeGenerator.cpp
+++ b/Engine/Source/Programs/UnrealHeaderTool/Private/CodeGenerator.cpp
@@ -2973,34 +2973,6 @@ void WriteMacro(FOutputDevice& Output, const FString& MacroName, FString MacroCo
 	Output.Log(Macroize(*MacroName, MoveTemp(MacroContent)));
 }
 
-static FString PrivatePropertiesOffsetGetters(const UStruct* Struct, const FString& StructCppName)
-{
-	check(Struct);
-
-	FUHTStringBuilder Result;
-	for (const FProperty* Property : TFieldRange<FProperty>(Struct, EFieldIteratorFlags::ExcludeSuper))
-	{
-		if (Property && Property->HasAnyPropertyFlags(CPF_NativeAccessSpecifierPrivate | CPF_NativeAccessSpecifierProtected) && !Property->HasAnyPropertyFlags(CPF_EditorOnly))
-		{
-			const FBoolProperty* BoolProperty = CastField<const FBoolProperty>(Property);
-			if (BoolProperty && !BoolProperty->IsNativeBool()) // if it's a bitfield
-			{
-				continue;
-			}
-
-			FString PropertyName = Property->GetName();
-			if (Property->HasAllPropertyFlags(CPF_Deprecated))
-			{
-				PropertyName += TEXT("_DEPRECATED");
-			}
-			Result.Logf(TEXT("\tFORCEINLINE static uint32 __PPO__%s() { return STRUCT_OFFSET(%s, %s); }") LINE_TERMINATOR,
-				*PropertyName, *StructCppName, *PropertyName);
-		}
-	}
-
-	return MoveTemp(Result);
-}
-
 void FNativeClassHeaderGenerator::ExportClassFromSourceFileInner(
 	FOutputDevice&           OutGeneratedHeaderText,
 	FOutputDevice&           OutCpp,
@@ -3288,12 +3260,6 @@ void FNativeClassHeaderGenerator::ExportClassFromSourceFileInner(
 
 				ExportConstructorsMacros(OutGeneratedHeaderText, OutCpp, StandardUObjectConstructorsMacroCall, EnhancedUObjectConstructorsMacroCall, SourceFile.GetGeneratedMacroName(ClassData), Class, *APIArg);
 			}
-			{
-				const FString PrivatePropertiesOffsets = PrivatePropertiesOffsetGetters(Class, ClassCPPName);
-				const FString PPOMacroNameRaw = SourceFile.GetGeneratedMacroName(ClassData, TEXT("_PRIVATE_PROPERTY_OFFSET"));
-				PPOMacroName = FString::Printf(TEXT("\t%s\r\n"), *PPOMacroNameRaw);
-				WriteMacro(OutGeneratedHeaderText, PPOMacroNameRaw, PrivatePropertiesOffsets);
-			}
 		}
 	}
 
@@ -3326,8 +3292,8 @@ void FNativeClassHeaderGenerator::ExportClassFromSourceFileInner(
 			MacroName = TEXT("GENERATED_UCLASS_BODY()");
 			DeprecationWarning = GetGeneratedMacroDeprecationWarning(MacroName);
 			GeneratedBodyLine = ClassData->GetGeneratedBodyLine();
-			LegacyGeneratedBody = FString::Printf(TEXT("%s%s%s"), *PPOMacroName, *ClassMacroCalls, *StandardUObjectConstructorsMacroCall);
-			GeneratedBody = FString::Printf(TEXT("%s%s%s"), *PPOMacroName, *ClassNoPureDeclsMacroCalls, *EnhancedUObjectConstructorsMacroCall);
+			LegacyGeneratedBody = FString::Printf(TEXT("%s%s"), *ClassMacroCalls, *StandardUObjectConstructorsMacroCall);
+			GeneratedBody = FString::Printf(TEXT("%s%s"), *ClassNoPureDeclsMacroCalls, *EnhancedUObjectConstructorsMacroCall);
 		}
 
 		FString WrappedLegacyGeneratedBody = FString::Printf(TEXT("%s%s%s%s%s%s"), *DeprecationWarning, DisableDeprecationWarnings, Public, *LegacyGeneratedBody, Public, EnableDeprecationWarnings);
@@ -3834,7 +3800,6 @@ void FNativeClassHeaderGenerator::ExportGeneratedStructBodyMacros(FOutputDevice&
 
 		const FString FriendLine = FString::Printf(TEXT("\tfriend struct %s_Statics;\r\n"), *ChoppedSingletonName);
 		const FString StaticClassLine = FString::Printf(TEXT("\t%sstatic class UScriptStruct* StaticStruct();\r\n"), (bRequiredAPI ? *FriendApiString : TEXT("")));
-		const FString PrivatePropertiesOffset = PrivatePropertiesOffsetGetters(Struct, StructNameCPP);
 		
 		// if we have RigVM methods on this struct we need to 
 		// declare the static method as well as the stub method
@@ -3891,7 +3856,7 @@ void FNativeClassHeaderGenerator::ExportGeneratedStructBodyMacros(FOutputDevice&
 
 		const FString SuperTypedef = BaseStruct ? FString::Printf(TEXT("\ttypedef %s Super;\r\n"), *FNameLookupCPP::GetNameCPP(BaseStruct)) : FString();
 
-		FString CombinedLine = FString::Printf(TEXT("%s%s%s%s%s"), *FriendLine, *StaticClassLine, *RigVMMethodsDeclarations, *PrivatePropertiesOffset, *SuperTypedef);
+		FString CombinedLine = FString::Printf(TEXT("%s%s%s%s"), *FriendLine, *StaticClassLine, *RigVMMethodsDeclarations, *SuperTypedef);
 		const FString MacroName = SourceFile.GetGeneratedBodyMacroName(Struct->StructMacroDeclaredLineNumber);
 
 		const FString Macroized = Macroize(*MacroName, MoveTemp(CombinedLine));
